Prevent crash when accessing the localization settings more than once.

#ttp 334601 - Crash: Live: Setting and Unsetting Use Localized Field Name

#branch UE4

#change Model wasn't disassociating itself with the OnSettingsChanged delegate when it was destroyed... which meant next time the settings were changed the delegate was being called with a stale pointer.

diff --git a/Engine/Source/Developer/InternationalizationSettings/Private/InternationalizationSettingsModelDetails.cpp b/Engine/Source/Developer/InternationalizationSettings/Private/InternationalizationSettingsModelDetails.cpp
--- a/Engine/Source/Developer/InternationalizationSettings/Private/InternationalizationSettingsModelDetails.cpp
+++ b/Engine/Source/Developer/InternationalizationSettings/Private/InternationalizationSettingsModelDetails.cpp
@@ -65,12 +65,15 @@ TSharedRef<IDetailCustomization> FInternationalizationSettingsModelDetails::Make
 
 FInternationalizationSettingsModelDetails::~FInternationalizationSettingsModelDetails()
 {
+	check(Model.IsValid());
+	Model->OnSettingChanged().RemoveAll(this);
 }
 
 void FInternationalizationSettingsModelDetails::OnSettingsChanged()
 {
 	FInternationalization& I18N = FInternationalization::Get();
 
+	check(Model.IsValid());
 	FString SavedCultureName = Model->GetCultureName();
 	if ( !SavedCultureName.IsEmpty() && SavedCultureName != I18N.GetCurrentCulture()->GetName() )
 	{
