Fixed locales/cultures that don't have base language localization data being omitted from region & language drop downs. Not absolutely perfect - could probably use a bigger refactor.

diff --git a/Engine/Source/Developer/InternationalizationSettings/Private/InternationalizationSettingsModelDetails.cpp b/Engine/Source/Developer/InternationalizationSettings/Private/InternationalizationSettingsModelDetails.cpp
--- a/Engine/Source/Developer/InternationalizationSettings/Private/InternationalizationSettingsModelDetails.cpp
+++ b/Engine/Source/Developer/InternationalizationSettings/Private/InternationalizationSettingsModelDetails.cpp
@@ -277,11 +277,11 @@ void FInternationalizationSettingsModelDetails::RefreshAvailableLanguages()
 	// Setup the language list
 	for( const auto& Culture : AvailableCultures )
 	{
-		const FString CultureRegionName = Culture->GetNativeRegion();
-		if ( CultureRegionName.IsEmpty() )
+		FCulturePtr LanguageCulture = FInternationalization::Get().GetCulture(Culture->GetTwoLetterISOLanguageName());
+		if (LanguageCulture.IsValid() && AvailableLanguages.Find(LanguageCulture) == INDEX_NONE)
 		{
-			AvailableLanguages.Add( Culture );
-
+			AvailableLanguages.Add(LanguageCulture);
+				
 			// Do we have a match for the base language
 			const FString CultureLanguageName = Culture->GetNativeLanguage();
 			if ( SelectedLanguageName == CultureLanguageName)
@@ -300,15 +300,15 @@ void FInternationalizationSettingsModelDetails::RefreshAvailableRegions()
 {
 	AvailableRegions.Empty();
 
-	FCulturePtr DefaultCulture = NULL;
+	FCulturePtr DefaultCulture;
 	if ( SelectedLanguage.IsValid() )
 	{
-		const FString SelectedLanguageName = SelectedLanguage->GetNativeLanguage();
+		const FString SelectedLanguageName = SelectedLanguage->GetTwoLetterISOLanguageName();
 
 		// Setup the region list
 		for( const auto& Culture : AvailableCultures )
 		{
-			const FString CultureLanguageName = Culture->GetNativeLanguage();
+			const FString CultureLanguageName = Culture->GetTwoLetterISOLanguageName();
 			if ( SelectedLanguageName == CultureLanguageName)
 			{
 				AvailableRegions.Add( Culture );
@@ -326,14 +326,21 @@ void FInternationalizationSettingsModelDetails::RefreshAvailableRegions()
 	}
 
 	// If we have a preferred default (or there's only one in the list), select that now
-	if ( DefaultCulture.IsValid() || AvailableCultures.Num() == 1 )
+	if ( !DefaultCulture.IsValid() )
+	{
+		if(AvailableRegions.Num() == 1)
+		{
+			DefaultCulture = AvailableRegions.Last();
+		}
+	}
+
+	if ( DefaultCulture.IsValid() )
 	{
-		FCulturePtr Culture = DefaultCulture.IsValid() ? DefaultCulture : AvailableCultures.Last();
 		// Set it as our default region, if one hasn't already been chosen
 		if ( !SelectedCulture.IsValid() && RegionComboBox.IsValid() )
 		{
 			// We have to update the combo box like this, otherwise it'll do a null selection when we next click on it
-			RegionComboBox->SetSelectedItem( Culture );
+			RegionComboBox->SetSelectedItem( DefaultCulture );
 		}
 	}
 
