In "Region and Language" you can toggle localization of node and pins in all graph editors.

SNodeTitle no longer compares cached node titles to the title the node wants to be to determine a reconstruction, it checks if the schema wants to purge all cached node titles.

diff --git a/Engine/Source/Developer/InternationalizationSettings/Classes/InternationalizationSettingsModel.h b/Engine/Source/Developer/InternationalizationSettings/Classes/InternationalizationSettingsModel.h
--- a/Engine/Source/Developer/InternationalizationSettings/Classes/InternationalizationSettingsModel.h
+++ b/Engine/Source/Developer/InternationalizationSettings/Classes/InternationalizationSettingsModel.h
@@ -22,7 +22,8 @@ public:
 	void SetCultureName(const FString& CultureName);
 	bool ShouldLoadLocalizedPropertyNames() const;
 	void ShouldLoadLocalizedPropertyNames(const bool Value);
-
+	bool ShouldShowNodesAndPinsUnlocalized() const;
+	void ShouldShowNodesAndPinsUnlocalized(const bool Value);
 public:
 
 	/**
diff --git a/Engine/Source/Developer/InternationalizationSettings/Private/InternationalizationSettingsModel.cpp b/Engine/Source/Developer/InternationalizationSettings/Private/InternationalizationSettingsModel.cpp
--- a/Engine/Source/Developer/InternationalizationSettings/Private/InternationalizationSettingsModel.cpp
+++ b/Engine/Source/Developer/InternationalizationSettings/Private/InternationalizationSettingsModel.cpp
@@ -20,6 +20,10 @@ void UInternationalizationSettingsModel::SaveDefaults()
 	bool bShouldLoadLocalizedPropertyNames = true;
 	GConfig->GetBool( TEXT("Internationalization"), TEXT("ShouldLoadLocalizedPropertyNames"), bShouldLoadLocalizedPropertyNames, GEditorGameAgnosticIni );
 	GConfig->SetBool( TEXT("Internationalization"), TEXT("ShouldLoadLocalizedPropertyNames"), bShouldLoadLocalizedPropertyNames, GEngineIni );
+
+	bool bShowNodesAndPinsUnlocalized = false;
+	GConfig->GetBool( TEXT("Internationalization"), TEXT("ShowNodesAndPinsUnlocalized"), bShowNodesAndPinsUnlocalized, GEditorGameAgnosticIni );
+	GConfig->SetBool( TEXT("Internationalization"), TEXT("ShowNodesAndPinsUnlocalized"), bShowNodesAndPinsUnlocalized, GEngineIni );
 }
 
 
@@ -33,6 +37,10 @@ void UInternationalizationSettingsModel::ResetToDefault()
 	GConfig->GetBool( TEXT("Internationalization"), TEXT("ShouldLoadLocalizedPropertyNames"), bShouldLoadLocalizedPropertyNames, GEngineIni );
 	GConfig->SetBool( TEXT("Internationalization"), TEXT("ShouldLoadLocalizedPropertyNames"), bShouldLoadLocalizedPropertyNames, GEditorGameAgnosticIni );
 
+	bool bShowNodesAndPinsUnlocalized = false;
+	GConfig->GetBool( TEXT("Internationalization"), TEXT("ShowNodesAndPinsUnlocalized"), bShowNodesAndPinsUnlocalized, GEngineIni );
+	GConfig->SetBool( TEXT("Internationalization"), TEXT("ShowNodesAndPinsUnlocalized"), bShowNodesAndPinsUnlocalized, GEditorGameAgnosticIni );
+
 	SettingChangedEvent.Broadcast();
 }
 
@@ -71,3 +79,18 @@ void UInternationalizationSettingsModel::ShouldLoadLocalizedPropertyNames(const
 	GConfig->SetBool( TEXT("Internationalization"), TEXT("ShouldLoadLocalizedPropertyNames"), Value, GEditorGameAgnosticIni );
 	SettingChangedEvent.Broadcast();
 }
+
+bool UInternationalizationSettingsModel::ShouldShowNodesAndPinsUnlocalized() const
+{
+	bool bShowNodesAndPinsUnlocalized = false;
+	if( !GConfig->GetBool( TEXT("Internationalization"), TEXT("ShowNodesAndPinsUnlocalized"), bShowNodesAndPinsUnlocalized, GEditorGameAgnosticIni ) )
+	{
+		GConfig->GetBool( TEXT("Internationalization"), TEXT("ShowNodesAndPinsUnlocalized"), bShowNodesAndPinsUnlocalized, GEngineIni );
+	}
+	return bShowNodesAndPinsUnlocalized;
+}
+
+void UInternationalizationSettingsModel::ShouldShowNodesAndPinsUnlocalized(const bool Value)
+{
+	GConfig->SetBool( TEXT("Internationalization"), TEXT("ShowNodesAndPinsUnlocalized"), Value, GEditorGameAgnosticIni );
+}
diff --git a/Engine/Source/Developer/InternationalizationSettings/Private/InternationalizationSettingsModelDetails.cpp b/Engine/Source/Developer/InternationalizationSettings/Private/InternationalizationSettingsModelDetails.cpp
--- a/Engine/Source/Developer/InternationalizationSettings/Private/InternationalizationSettingsModelDetails.cpp
+++ b/Engine/Source/Developer/InternationalizationSettings/Private/InternationalizationSettingsModelDetails.cpp
@@ -1,6 +1,7 @@
 ﻿// Copyright 1998-2015 Epic Games, Inc. All Rights Reserved.
 
 #include "InternationalizationSettingsModulePrivatePCH.h"
+#include "EdGraph/EdGraphSchema.h"
 
 #define LOCTEXT_NAMESPACE "InternationalizationSettingsModelDetails"
 
@@ -217,7 +218,32 @@ void FInternationalizationSettingsModelDetails::CustomizeDetails( IDetailLayoutB
 		SAssignNew(FieldNamesCheckBox, SCheckBox)
 		.IsChecked(Model->ShouldLoadLocalizedPropertyNames() ? ECheckBoxState::Checked : ECheckBoxState::Unchecked)
 		.ToolTipText(FieldNamesToolTipText)
-		.OnCheckStateChanged(this, &FInternationalizationSettingsModelDetails::ShoudLoadLocalizedFieldNamesCheckChanged)
+		.OnCheckStateChanged(this, &FInternationalizationSettingsModelDetails::ShouldLoadLocalizedFieldNamesCheckChanged)
+	];
+
+	const FText NodeAndPinsNamesToolTipText = LOCTEXT("GraphEditorNodesAndPinsLocalized_Tooltip", "Toggle localized node and pin titles in all graph editors");
+
+	CategoryBuilder.AddCustomRow(LOCTEXT("GraphEditorNodesAndPinsLocalized", "Use Localized Graph Editor Nodes and Pins"))
+	.NameContent()
+	[
+		SNew(SHorizontalBox)
+		+SHorizontalBox::Slot()
+		.Padding( FMargin( 0, 1, 0, 1 ) )
+		.FillWidth(1.0f)
+		[
+			SNew(STextBlock)
+			.Text(LOCTEXT("GraphEditorNodesAndPinsLocalized", "Use Localized Graph Editor Nodes and Pins"))
+			.Font(DetailBuilder.GetDetailFont())
+			.ToolTipText(NodeAndPinsNamesToolTipText)
+		]
+	]
+	.ValueContent()
+	.MaxDesiredWidth(300.0f)
+	[
+		SAssignNew(FieldNamesCheckBox, SCheckBox)
+		.IsChecked(Model->ShouldShowNodesAndPinsUnlocalized() ? ECheckBoxState::Unchecked : ECheckBoxState::Checked)
+		.ToolTipText(NodeAndPinsNamesToolTipText)
+		.OnCheckStateChanged(this, &FInternationalizationSettingsModelDetails::ShouldShowNodesAndPinsUnlocalized)
 	];
 
 	CategoryBuilder.AddCustomRow(LOCTEXT("EditorRestartWarningLabel", "RestartWarning"))
@@ -408,7 +434,7 @@ EVisibility FInternationalizationSettingsModelDetails::GetInternationalizationRe
 	return RequiresRestart ? EVisibility::Visible : EVisibility::Collapsed;
 }
 
-void FInternationalizationSettingsModelDetails::ShoudLoadLocalizedFieldNamesCheckChanged(ECheckBoxState CheckState)
+void FInternationalizationSettingsModelDetails::ShouldLoadLocalizedFieldNamesCheckChanged(ECheckBoxState CheckState)
 {
 	HandleShutdownPostPackagesSaved();
 }
@@ -432,4 +458,21 @@ void FInternationalizationSettingsModelDetails::HandleShutdownPostPackagesSaved(
 	}
 }
 
+void FInternationalizationSettingsModelDetails::ShouldShowNodesAndPinsUnlocalized(ECheckBoxState CheckState)
+{
+	check(Model.IsValid());
+	Model->ShouldShowNodesAndPinsUnlocalized(CheckState == ECheckBoxState::Unchecked);
+
+	// Find all Schemas and force a visualization cache clear
+	for ( TObjectIterator<UClass> ClassIt; ClassIt; ++ClassIt )
+	{
+		UClass* CurrentClass = *ClassIt;
+
+		if (UEdGraphSchema* Schema = Cast<UEdGraphSchema>(CurrentClass->GetDefaultObject()))
+		{
+			Schema->ForceVisualizationCacheClear();
+		}
+	}
+}
+
 #undef LOCTEXT_NAMESPACE
diff --git a/Engine/Source/Developer/InternationalizationSettings/Public/InternationalizationSettingsModelDetails.h b/Engine/Source/Developer/InternationalizationSettings/Public/InternationalizationSettingsModelDetails.h
--- a/Engine/Source/Developer/InternationalizationSettings/Public/InternationalizationSettingsModelDetails.h
+++ b/Engine/Source/Developer/InternationalizationSettings/Public/InternationalizationSettingsModelDetails.h
@@ -63,11 +63,14 @@ private:
 	EVisibility GetInternationalizationRestartRowVisibility() const;
 
 	/** Delegate called when the the checked state of whether or not field names should be localized has changed. */
-	void ShoudLoadLocalizedFieldNamesCheckChanged(ECheckBoxState CheckState);
+	void ShouldLoadLocalizedFieldNamesCheckChanged(ECheckBoxState CheckState);
 
 	/** Write to config now the Editor is shutting down (all packages are saved) */
 	void HandleShutdownPostPackagesSaved();
 
+	/** Delegate called when the the checked state of whether or not nodes and pins in graph editors should be localized has changed. */
+	void ShouldShowNodesAndPinsUnlocalized(ECheckBoxState CheckState);
+
 private:
 	TWeakObjectPtr<UInternationalizationSettingsModel> Model;
 
diff --git a/Engine/Source/Editor/BehaviorTreeEditor/Classes/EdGraphSchema_BehaviorTree.h b/Engine/Source/Editor/BehaviorTreeEditor/Classes/EdGraphSchema_BehaviorTree.h
--- a/Engine/Source/Editor/BehaviorTreeEditor/Classes/EdGraphSchema_BehaviorTree.h
+++ b/Engine/Source/Editor/BehaviorTreeEditor/Classes/EdGraphSchema_BehaviorTree.h
@@ -37,9 +37,16 @@ class UEdGraphSchema_BehaviorTree : public UAIGraphSchema
 	virtual FLinearColor GetPinTypeColor(const FEdGraphPinType& PinType) const override;
 	virtual class FConnectionDrawingPolicy* CreateConnectionDrawingPolicy(int32 InBackLayerID, int32 InFrontLayerID, float InZoomFactor, const FSlateRect& InClippingRect, class FSlateWindowElementList& InDrawElements, class UEdGraph* InGraphObj) const override;
 	virtual int32 GetNodeSelectionCount(const UEdGraph* Graph) const override;
+	virtual bool IsCacheVisualizationOutOfDate(int32 InVisualizationCacheID) const override;
+	virtual int32 GetCurrentVisualizationCacheID() const override;
+	virtual void ForceVisualizationCacheClear() const override;
 	// End EdGraphSchema interface
 
 	virtual void GetGraphNodeContextActions(FGraphContextMenuBuilder& ContextMenuBuilder, int32 SubNodeFlags) const override;
 	virtual void GetSubNodeClasses(int32 SubNodeFlags, TArray<FGraphNodeClassData>& ClassData, UClass*& GraphNodeClass) const override;
+
+private:
+	// ID for checking dirty status of node titles against, increases whenever 
+	static int32 CurrentCacheRefreshID;
 };
 
diff --git a/Engine/Source/Editor/BehaviorTreeEditor/Private/BehaviorTreeEditor.cpp b/Engine/Source/Editor/BehaviorTreeEditor/Private/BehaviorTreeEditor.cpp
--- a/Engine/Source/Editor/BehaviorTreeEditor/Private/BehaviorTreeEditor.cpp
+++ b/Engine/Source/Editor/BehaviorTreeEditor/Private/BehaviorTreeEditor.cpp
@@ -1120,6 +1120,7 @@ void FBehaviorTreeEditor::OnFinishedChangingProperties(const FPropertyChangedEve
 		MyGraph->UpdateInjectedNodes();
 		MyGraph->UpdateAsset(UBehaviorTreeGraph::ClearDebuggerFlags);
 	}
+	BehaviorTree->BTGraph->GetSchema()->ForceVisualizationCacheClear();
 }
 
 void FBehaviorTreeEditor::OnPackageSaved(const FString& PackageFileName, UObject* Outer)
diff --git a/Engine/Source/Editor/BehaviorTreeEditor/Private/EdGraphSchema_BehaviorTree.cpp b/Engine/Source/Editor/BehaviorTreeEditor/Private/EdGraphSchema_BehaviorTree.cpp
--- a/Engine/Source/Editor/BehaviorTreeEditor/Private/EdGraphSchema_BehaviorTree.cpp
+++ b/Engine/Source/Editor/BehaviorTreeEditor/Private/EdGraphSchema_BehaviorTree.cpp
@@ -17,6 +17,8 @@
 
 #define LOCTEXT_NAMESPACE "BehaviorTreeEditor"
 
+int32 UEdGraphSchema_BehaviorTree::CurrentCacheRefreshID = 0;
+
 UEdGraphNode* FBehaviorTreeSchemaAction_AutoArrange::PerformAction(class UEdGraph* ParentGraph, UEdGraphPin* FromPin, const FVector2D Location, bool bSelectNewNode)
 {
 	UBehaviorTreeGraph* Graph = Cast<UBehaviorTreeGraph>(ParentGraph);
@@ -415,4 +417,19 @@ int32 UEdGraphSchema_BehaviorTree::GetNodeSelectionCount(const UEdGraph* Graph)
 	return 0;
 }
 
+bool UEdGraphSchema_BehaviorTree::IsCacheVisualizationOutOfDate(int32 InVisualizationCacheID) const
+{
+	return CurrentCacheRefreshID != InVisualizationCacheID;
+}
+
+int32 UEdGraphSchema_BehaviorTree::GetCurrentVisualizationCacheID() const
+{
+	return CurrentCacheRefreshID;
+}
+
+void UEdGraphSchema_BehaviorTree::ForceVisualizationCacheClear() const
+{
+	++CurrentCacheRefreshID;
+}
+
 #undef LOCTEXT_NAMESPACE
diff --git a/Engine/Source/Editor/BlueprintGraph/Classes/K2Node_CustomEvent.h b/Engine/Source/Editor/BlueprintGraph/Classes/K2Node_CustomEvent.h
--- a/Engine/Source/Editor/BlueprintGraph/Classes/K2Node_CustomEvent.h
+++ b/Engine/Source/Editor/BlueprintGraph/Classes/K2Node_CustomEvent.h
@@ -24,7 +24,6 @@ class UK2Node_CustomEvent : public UK2Node_Event
 	// Begin UEdGraphNode interface
 	virtual void ReconstructNode() override;
 	virtual FText GetNodeTitle(ENodeTitleType::Type TitleType) const override;
-	virtual void RefreshNodeTitle() override;
 	virtual void OnRenameNode(const FString& NewName) override;
 	virtual TSharedPtr<class INameValidatorInterface> MakeNameValidator() const override;
 	virtual FText GetTooltipText() const override;
diff --git a/Engine/Source/Editor/BlueprintGraph/Private/K2Node_CustomEvent.cpp b/Engine/Source/Editor/BlueprintGraph/Private/K2Node_CustomEvent.cpp
--- a/Engine/Source/Editor/BlueprintGraph/Private/K2Node_CustomEvent.cpp
+++ b/Engine/Source/Editor/BlueprintGraph/Private/K2Node_CustomEvent.cpp
@@ -137,11 +137,6 @@ FText UK2Node_CustomEvent::GetNodeTitle(ENodeTitleType::Type TitleType) const
 	return CachedNodeTitle;
 }
 
-void UK2Node_CustomEvent::RefreshNodeTitle()
-{
-	CachedNodeTitle.MarkDirty();
-}
-
 bool UK2Node_CustomEvent::CanCreateUserDefinedPin(const FEdGraphPinType& InPinType, EEdGraphPinDirection InDesiredDirection, FText& OutErrorMessage)
 {
 	if (!IsEditable())
diff --git a/Engine/Source/Editor/GraphEditor/Private/SGraphNode.cpp b/Engine/Source/Editor/GraphEditor/Private/SGraphNode.cpp
--- a/Engine/Source/Editor/GraphEditor/Private/SGraphNode.cpp
+++ b/Engine/Source/Editor/GraphEditor/Private/SGraphNode.cpp
@@ -31,16 +31,16 @@ void SNodeTitle::Construct(const FArguments& InArgs, UEdGraphNode* InNode)
 	{
 		TitleText = TAttribute<FText>(this, &SNodeTitle::GetNodeTitle);
 	}
-	CachedTitle = TitleText.Get();
+	NodeTitleCache.SetCachedText(TitleText.Get(), GraphNode);
 	RebuildWidget();
 }
 
 void SNodeTitle::Tick( const FGeometry& AllottedGeometry, const double InCurrentTime, const float InDeltaTime )
 {
 	// Checks to see if the cached string is valid, and if not, updates it.
-	if (TitleText.Get().CompareTo(CachedTitle) != 0)
+	if (NodeTitleCache.IsOutOfDate(GraphNode))
 	{
-		CachedTitle = TitleText.Get();
+		NodeTitleCache.SetCachedText(TitleText.Get(), GraphNode);
 		RebuildWidget();
 	}
 }
@@ -68,7 +68,7 @@ void SNodeTitle::RebuildWidget()
 
 	// Break the title into lines
 	TArray<FString> Lines;
-	const FString CachedTitleString = CachedTitle.ToString().Replace(TEXT("\r"), TEXT(""));
+	const FString CachedTitleString = NodeTitleCache.GetCachedText().ToString().Replace(TEXT("\r"), TEXT(""));
 	CachedTitleString.ParseIntoArray(Lines, TEXT("\n"), false);
 
 	if (Lines.Num())
diff --git a/Engine/Source/Editor/GraphEditor/Public/SGraphNode.h b/Engine/Source/Editor/GraphEditor/Public/SGraphNode.h
--- a/Engine/Source/Editor/GraphEditor/Public/SGraphNode.h
+++ b/Engine/Source/Editor/GraphEditor/Public/SGraphNode.h
@@ -41,7 +41,7 @@ public:
 
 protected:
 	UEdGraphNode* GraphNode;
-	FText CachedTitle;
+	FNodeTextCache NodeTitleCache;
 	FName ExtraLineStyle;
 
 	/** The cached head title to return */
diff --git a/Engine/Source/Editor/Kismet/Private/BlueprintDetailsCustomization.cpp b/Engine/Source/Editor/Kismet/Private/BlueprintDetailsCustomization.cpp
--- a/Engine/Source/Editor/Kismet/Private/BlueprintDetailsCustomization.cpp
+++ b/Engine/Source/Editor/Kismet/Private/BlueprintDetailsCustomization.cpp
@@ -2694,9 +2694,6 @@ void FBlueprintGraphActionDetails::SetNetFlags( TWeakObjectPtr<UK2Node_EditableP
 				bBlueprintModified = true;
 			}
 
-			// The node title needs to refresh
-			FunctionEntryNode->RefreshNodeTitle();
-
 			if( bBlueprintModified )
 			{
 				FBlueprintEditorUtils::MarkBlueprintAsStructurallyModified( FunctionEntryNode->GetBlueprint() );
diff --git a/Engine/Source/Editor/Kismet/Private/SBlueprintPalette.cpp b/Engine/Source/Editor/Kismet/Private/SBlueprintPalette.cpp
--- a/Engine/Source/Editor/Kismet/Private/SBlueprintPalette.cpp
+++ b/Engine/Source/Editor/Kismet/Private/SBlueprintPalette.cpp
@@ -1046,25 +1046,29 @@ TSharedRef<SWidget> SBlueprintPaletteItem::CreateTextSlotWidget(const FSlateFont
 //------------------------------------------------------------------------------
 FText SBlueprintPaletteItem::GetDisplayText() const
 {
-	FText DisplayText;
-
-	TSharedPtr< FEdGraphSchemaAction > GraphAction = ActionPtr.Pin();
-	if (GraphAction->GetTypeId() == FEdGraphSchemaAction_K2Enum::StaticGetTypeId())
-	{
-		FEdGraphSchemaAction_K2Enum* EnumAction = (FEdGraphSchemaAction_K2Enum*)GraphAction.Get();
-		DisplayText = FText::FromString(EnumAction->Enum->GetName());
-	}
-	else if (GraphAction->GetTypeId() == FEdGraphSchemaAction_K2Struct::StaticGetTypeId())
+	const UEdGraphSchema_K2* K2Schema = GetDefault<UEdGraphSchema_K2>();
+	if (MenuDescriptionCache.IsOutOfDate(K2Schema))
 	{
-		FEdGraphSchemaAction_K2Struct* StructAction = (FEdGraphSchemaAction_K2Struct*)GraphAction.Get();
-		DisplayText = StructAction->Struct ? FText::FromString(StructAction->Struct->GetName()) : FText::FromString(TEXT("None"));
-	}
-	else
-	{
-		DisplayText = ActionPtr.Pin()->MenuDescription;
+		TSharedPtr< FEdGraphSchemaAction > GraphAction = ActionPtr.Pin();
+		if (GraphAction->GetTypeId() == FEdGraphSchemaAction_K2Enum::StaticGetTypeId())
+		{
+			FEdGraphSchemaAction_K2Enum* EnumAction = (FEdGraphSchemaAction_K2Enum*)GraphAction.Get();
+			FText DisplayText = FText::FromString(EnumAction->Enum->GetName());
+			MenuDescriptionCache.SetCachedText(DisplayText, K2Schema);
+		}
+		else if (GraphAction->GetTypeId() == FEdGraphSchemaAction_K2Struct::StaticGetTypeId())
+		{
+			FEdGraphSchemaAction_K2Struct* StructAction = (FEdGraphSchemaAction_K2Struct*)GraphAction.Get();
+			FText DisplayText = StructAction->Struct ? FText::FromString(StructAction->Struct->GetName()) : FText::FromString(TEXT("None"));
+			MenuDescriptionCache.SetCachedText(DisplayText, K2Schema);
+		}
+		else
+		{
+			MenuDescriptionCache.SetCachedText(ActionPtr.Pin()->MenuDescription, K2Schema);
+		}
 	}
 
-	return DisplayText;
+	return MenuDescriptionCache;
 }
 
 //------------------------------------------------------------------------------
diff --git a/Engine/Source/Editor/Kismet/Private/SBlueprintPalette.h b/Engine/Source/Editor/Kismet/Private/SBlueprintPalette.h
--- a/Engine/Source/Editor/Kismet/Private/SBlueprintPalette.h
+++ b/Engine/Source/Editor/Kismet/Private/SBlueprintPalette.h
@@ -61,6 +61,9 @@ private:
 
 	/** Pointer back to the blueprint editor that owns this */
 	TWeakPtr<FBlueprintEditor> BlueprintEditorPtr;
+
+	/** Cache for the MenuDescription to be displayed for this item */
+	FNodeTextCache MenuDescriptionCache;
 };
 
 
diff --git a/Engine/Source/Editor/MaterialEditor/Private/MaterialEditor.cpp b/Engine/Source/Editor/MaterialEditor/Private/MaterialEditor.cpp
--- a/Engine/Source/Editor/MaterialEditor/Private/MaterialEditor.cpp
+++ b/Engine/Source/Editor/MaterialEditor/Private/MaterialEditor.cpp
@@ -3567,6 +3567,8 @@ void FMaterialEditor::NotifyPostChange( const FPropertyChangedEvent& PropertyCha
 
 	Material->MarkPackageDirty();
 	SetMaterialDirty();
+
+	GetDefault<UMaterialGraphSchema>()->ForceVisualizationCacheClear();
 }
 
 void FMaterialEditor::ToggleCollapsed(UMaterialExpression* MaterialExpression)
diff --git a/Engine/Source/Editor/Persona/Private/AnimGraphNodeDetails.cpp b/Engine/Source/Editor/Persona/Private/AnimGraphNodeDetails.cpp
--- a/Engine/Source/Editor/Persona/Private/AnimGraphNodeDetails.cpp
+++ b/Engine/Source/Editor/Persona/Private/AnimGraphNodeDetails.cpp
@@ -484,6 +484,12 @@ void FBoneReferenceCustomization::CustomizeChildren( TSharedRef<IPropertyHandle>
 void FBoneReferenceCustomization::OnBoneSelectionChanged(FName Name)
 {
 	BoneRefProperty->SetValue(Name);
+
+	// Force a refresh of cached node titles, K2Schema is no different than the AnimSchema for this purpose
+	if (const UEdGraphSchema_K2* K2Schema = GetDefault<UEdGraphSchema_K2>())
+	{
+		K2Schema->ForceVisualizationCacheClear();
+	}
 }
 FName FBoneReferenceCustomization::GetSelectedBone() const
 {
diff --git a/Engine/Source/Editor/UnrealEd/Classes/MaterialGraph/MaterialGraphSchema.h b/Engine/Source/Editor/UnrealEd/Classes/MaterialGraph/MaterialGraphSchema.h
--- a/Engine/Source/Editor/UnrealEd/Classes/MaterialGraph/MaterialGraphSchema.h
+++ b/Engine/Source/Editor/UnrealEd/Classes/MaterialGraph/MaterialGraphSchema.h
@@ -203,6 +203,9 @@ class UMaterialGraphSchema : public UEdGraphSchema
 	virtual void DroppedAssetsOnGraph(const TArray<class FAssetData>& Assets, const FVector2D& GraphPosition, UEdGraph* Graph) const override;
 	virtual int32 GetNodeSelectionCount(const UEdGraph* Graph) const override;
 	virtual TSharedPtr<FEdGraphSchemaAction> GetCreateCommentAction() const override;
+	virtual bool IsCacheVisualizationOutOfDate(int32 InVisualizationCacheID) const override;
+	virtual int32 GetCurrentVisualizationCacheID() const override;
+	virtual void ForceVisualizationCacheClear() const override;
 	// End UEdGraphSchema interface
 
 private:
@@ -219,4 +222,7 @@ private:
 	*/
 	bool HasCompatibleConnection(const FAssetData& FunctionAssetData, uint32 TestType, EEdGraphPinDirection TestDirection) const;
 
+private:
+	// ID for checking dirty status of node titles against, increases whenever 
+	static int32 CurrentCacheRefreshID;
 };
diff --git a/Engine/Source/Editor/UnrealEd/Private/MaterialGraphSchema.cpp b/Engine/Source/Editor/UnrealEd/Private/MaterialGraphSchema.cpp
--- a/Engine/Source/Editor/UnrealEd/Private/MaterialGraphSchema.cpp
+++ b/Engine/Source/Editor/UnrealEd/Private/MaterialGraphSchema.cpp
@@ -22,6 +22,8 @@
 
 #define LOCTEXT_NAMESPACE "MaterialGraphSchema"
 
+int32 UMaterialGraphSchema::CurrentCacheRefreshID = 0;
+
 ////////////////////////////////////////
 // FMaterialGraphSchemaAction_NewNode //
 
@@ -809,4 +811,19 @@ bool UMaterialGraphSchema::HasCompatibleConnection(const FAssetData& FunctionAss
 	return false;
 }
 
+bool UMaterialGraphSchema::IsCacheVisualizationOutOfDate(int32 InVisualizationCacheID) const
+{
+	return CurrentCacheRefreshID != InVisualizationCacheID;
+}
+
+int32 UMaterialGraphSchema::GetCurrentVisualizationCacheID() const
+{
+	return CurrentCacheRefreshID;
+}
+
+void UMaterialGraphSchema::ForceVisualizationCacheClear() const
+{
+	++CurrentCacheRefreshID;
+}
+
 #undef LOCTEXT_NAMESPACE
diff --git a/Engine/Source/Runtime/Engine/Classes/EdGraph/EdGraphNode.h b/Engine/Source/Runtime/Engine/Classes/EdGraph/EdGraphNode.h
--- a/Engine/Source/Runtime/Engine/Classes/EdGraph/EdGraphNode.h
+++ b/Engine/Source/Runtime/Engine/Classes/EdGraph/EdGraphNode.h
@@ -285,9 +285,6 @@ public:
 
 	/** Gets the name of this node, shown in title bar */
 	virtual FText GetNodeTitle(ENodeTitleType::Type TitleType) const;
-	
-	/** Forces a refresh of the node title if it's cached */
-	virtual void RefreshNodeTitle() {};
 
 	/** 
 	 * Gets the draw color of a node's title bar
diff --git a/Engine/Source/Runtime/Engine/Private/EdGraph/EdGraphPin.cpp b/Engine/Source/Runtime/Engine/Private/EdGraph/EdGraphPin.cpp
--- a/Engine/Source/Runtime/Engine/Private/EdGraph/EdGraphPin.cpp
+++ b/Engine/Source/Runtime/Engine/Private/EdGraph/EdGraphPin.cpp
@@ -262,6 +262,13 @@ FText UEdGraphPin::GetDisplayName() const
 	else
 	{
 		DisplayName = (!PinFriendlyName.IsEmpty()) ? PinFriendlyName : FText::FromString(PinName);
+
+		bool bShowNodesAndPinsUnlocalized;
+		GConfig->GetBool( TEXT("Internationalization"), TEXT("ShowNodesAndPinsUnlocalized"), bShowNodesAndPinsUnlocalized, GEditorGameAgnosticIni );
+		if (bShowNodesAndPinsUnlocalized)
+		{
+			return FText::FromString(DisplayName.BuildSourceString());
+		}
 	}
 	return DisplayName;
 }
diff --git a/Engine/Source/Runtime/Engine/Private/EdGraph/EdGraphSchema.cpp b/Engine/Source/Runtime/Engine/Private/EdGraph/EdGraphSchema.cpp
--- a/Engine/Source/Runtime/Engine/Private/EdGraph/EdGraphSchema.cpp
+++ b/Engine/Source/Runtime/Engine/Private/EdGraph/EdGraphSchema.cpp
@@ -491,8 +491,23 @@ FPinConnectionResponse UEdGraphSchema::CopyPinLinks(UEdGraphPin& CopyFromPin, UE
 
 FText UEdGraphSchema::GetPinDisplayName(const UEdGraphPin* Pin) const
 {
+	FText ResultPinName;
 	check(Pin != NULL);
-	return !Pin->PinFriendlyName.IsEmpty() ? Pin->PinFriendlyName : FText::FromString(Pin->PinName);
+	if (Pin->PinFriendlyName.IsEmpty())
+	{
+		ResultPinName = FText::FromString(Pin->PinName);
+	}
+	else
+	{
+		ResultPinName = Pin->PinFriendlyName;
+		bool bShowNodesAndPinsUnlocalized;
+		GConfig->GetBool( TEXT("Internationalization"), TEXT("ShowNodesAndPinsUnlocalized"), bShowNodesAndPinsUnlocalized, GEditorGameAgnosticIni );
+		if (bShowNodesAndPinsUnlocalized)
+		{
+			ResultPinName = FText::FromString(ResultPinName.BuildSourceString());
+		}
+	}
+	return ResultPinName;
 }
 
 void UEdGraphSchema::ConstructBasicPinTooltip(UEdGraphPin const& Pin, FText const& PinDescription, FString& TooltipOut) const
diff --git a/Engine/Source/Runtime/Engine/Public/EdGraph/EdGraphNodeUtils.h b/Engine/Source/Runtime/Engine/Public/EdGraph/EdGraphNodeUtils.h
--- a/Engine/Source/Runtime/Engine/Public/EdGraph/EdGraphNodeUtils.h
+++ b/Engine/Source/Runtime/Engine/Public/EdGraph/EdGraphNodeUtils.h
@@ -27,8 +27,17 @@ public:
 	 */
 	FORCEINLINE bool IsOutOfDate(const UEdGraphNode* InOwningNode) const
 	{
-		check(InOwningNode);
-		return CachedText.IsEmpty() || (InOwningNode->GetSchema()->IsCacheVisualizationOutOfDate(CacheRefreshID));
+		return CachedText.IsEmpty() || (InOwningNode && InOwningNode->GetGraph() && InOwningNode->GetSchema()->IsCacheVisualizationOutOfDate(CacheRefreshID));
+	}
+
+	/**
+	 * Checks if the title is out of date
+	 *
+	 * @param InSchema		Schema to use for checking if externally made out of date
+	 */
+	FORCEINLINE bool IsOutOfDate(const UEdGraphSchema* InSchema) const
+	{
+		return CachedText.IsEmpty() || (InSchema && InSchema->IsCacheVisualizationOutOfDate(CacheRefreshID));
 	}
 
 	/**
@@ -39,9 +48,30 @@ public:
 	 */
 	FORCEINLINE void SetCachedText(FText const& InText, const UEdGraphNode* InOwningNode) const
 	{
-		check(InOwningNode);
-		CachedText = InText;
-		CacheRefreshID = InOwningNode->GetSchema()->GetCurrentVisualizationCacheID();
+		UpdateCacheIternal(InText);
+
+		if (InOwningNode && InOwningNode->GetGraph())
+		{
+			// Update the CacheRefreshID, whenever these values do not match the cached value is out of date
+			CacheRefreshID = InOwningNode->GetSchema()->GetCurrentVisualizationCacheID();
+		}
+	}
+
+	/**
+	 * Checks if the title is out of date
+	 *
+	 * @param InText			Text to cache
+	 * @param InSchema			Schema to use for checking if externally made out of date
+	 */
+	FORCEINLINE void SetCachedText(FText const& InText, const UEdGraphSchema* InSchema) const
+	{
+		UpdateCacheIternal(InText);
+
+		if (InSchema)
+		{
+			// Update the CacheRefreshID, whenever these values do not match the cached value is out of date
+			CacheRefreshID = InSchema->GetCurrentVisualizationCacheID();
+		}
 	}
 
 	/** */
@@ -69,6 +99,24 @@ public:
 		return GetCachedText();
 	}
 
+private:
+	/** Helper to lookup from the GConfig the correct ini setting for displaying
+	 *  node and pin titles localized and updates the cache accordingly */
+	FORCEINLINE bool UpdateCacheIternal(FText const& InText) const
+	{
+		bool bShowNodesAndPinsUnlocalized;
+		GConfig->GetBool( TEXT("Internationalization"), TEXT("ShowNodesAndPinsUnlocalized"), bShowNodesAndPinsUnlocalized, GEditorGameAgnosticIni );
+
+		if (bShowNodesAndPinsUnlocalized)
+		{
+			CachedText = FText::FromString(InText.BuildSourceString());
+		}
+		else
+		{
+			CachedText = InText;
+		}
+	}
+
 private:
 	/** Mutable so that SetCachedText() can remain const (callable by the node's GetNodeTitle() method)*/
 	mutable FText CachedText;
@@ -147,7 +195,8 @@ public:
 	 */
 	FORCEINLINE bool IsTooltipCached(const UEdGraphNode* InOwningNode) const
 	{
-		return !CachedTooltip.IsOutOfDate(InOwningNode);
+		const UEdGraphSchema* NullSchema = nullptr;
+		return !CachedTooltip.IsOutOfDate(NullSchema);
 	}
 
 	/**
@@ -158,7 +207,8 @@ public:
 	 */
 	FORCEINLINE void SetCachedTooltip(FText const& InText, const UEdGraphNode* InOwningNode) const
 	{
-		CachedTooltip.SetCachedText(InText, InOwningNode);
+		const UEdGraphSchema* NullSchema = nullptr;
+		CachedTooltip.SetCachedText(InText, NullSchema);
 	}
 
 	/** */
