Added game native culture setting to control which culture's localization data the editor loads for game assets. Game assets now display the specified culture's localization data when in the editor.

diff --git a/Engine/Source/Developer/InternationalizationSettings/Classes/InternationalizationSettingsModel.h b/Engine/Source/Developer/InternationalizationSettings/Classes/InternationalizationSettingsModel.h
--- a/Engine/Source/Developer/InternationalizationSettings/Classes/InternationalizationSettingsModel.h
+++ b/Engine/Source/Developer/InternationalizationSettings/Classes/InternationalizationSettingsModel.h
@@ -18,8 +18,10 @@ public:
 
 	void SaveDefaults();
 	void ResetToDefault();
-	FString GetCultureName() const;
-	void SetCultureName(const FString& CultureName);
+	FString GetEditorCultureName() const;
+	void SetEditorCultureName(const FString& CultureName);
+	FString GetNativeGameCultureName() const;
+	void SetNativeGameCultureName(const FString& CultureName);
 	bool ShouldLoadLocalizedPropertyNames() const;
 	void ShouldLoadLocalizedPropertyNames(const bool Value);
 	bool ShouldShowNodesAndPinsUnlocalized() const;
@@ -31,11 +33,11 @@ public:
 	 *
 	 * @return The delegate.
 	 */
-	DECLARE_EVENT(UInternationalizationSettingsModel, FSettingChangedEvent);
-	FSettingChangedEvent& OnSettingChanged( ) { return SettingChangedEvent; }
+	DECLARE_EVENT(UInternationalizationSettingsModel, FSettingsChangedEvent);
+	FSettingsChangedEvent& OnSettingsChanged( ) { return SettingsChangedEvent; }
 
 private:
 
 	// Holds an event delegate that is executed when a setting has changed.
-	FSettingChangedEvent SettingChangedEvent;
+	FSettingsChangedEvent SettingsChangedEvent;
 };
diff --git a/Engine/Source/Developer/InternationalizationSettings/Private/InternationalizationSettingsModel.cpp b/Engine/Source/Developer/InternationalizationSettings/Private/InternationalizationSettingsModel.cpp
--- a/Engine/Source/Developer/InternationalizationSettings/Private/InternationalizationSettingsModel.cpp
+++ b/Engine/Source/Developer/InternationalizationSettings/Private/InternationalizationSettingsModel.cpp
@@ -2,7 +2,6 @@
 
 #include "InternationalizationSettingsModulePrivatePCH.h"
 
-
 /* UInternationalizationSettingsModel interface
  *****************************************************************************/
 
@@ -10,7 +9,6 @@ UInternationalizationSettingsModel::UInternationalizationSettingsModel( const FO
 	: Super(ObjectInitializer)
 { }
 
-
 void UInternationalizationSettingsModel::SaveDefaults()
 {
 	FString SavedCultureName;
@@ -26,13 +24,14 @@ void UInternationalizationSettingsModel::SaveDefaults()
 	GConfig->SetBool( TEXT("Internationalization"), TEXT("ShowNodesAndPinsUnlocalized"), bShowNodesAndPinsUnlocalized, GEngineIni );
 }
 
-
 void UInternationalizationSettingsModel::ResetToDefault()
 {
 	FString SavedCultureName;
 	GConfig->GetString( TEXT("Internationalization"), TEXT("Culture"), SavedCultureName, GEngineIni );
 	GConfig->SetString( TEXT("Internationalization"), TEXT("Culture"), *SavedCultureName, GEditorSettingsIni );
 
+	GConfig->SetString( TEXT("Internationalization"), TEXT("NativeGameCulture"), TEXT(""), GEditorSettingsIni );
+
 	bool bShouldLoadLocalizedPropertyNames = true;
 	GConfig->GetBool( TEXT("Internationalization"), TEXT("ShouldLoadLocalizedPropertyNames"), bShouldLoadLocalizedPropertyNames, GEngineIni );
 	GConfig->SetBool( TEXT("Internationalization"), TEXT("ShouldLoadLocalizedPropertyNames"), bShouldLoadLocalizedPropertyNames, GEditorSettingsIni );
@@ -41,11 +40,10 @@ void UInternationalizationSettingsModel::ResetToDefault()
 	GConfig->GetBool( TEXT("Internationalization"), TEXT("ShowNodesAndPinsUnlocalized"), bShowNodesAndPinsUnlocalized, GEngineIni );
 	GConfig->SetBool( TEXT("Internationalization"), TEXT("ShowNodesAndPinsUnlocalized"), bShowNodesAndPinsUnlocalized, GEditorSettingsIni );
 
-	SettingChangedEvent.Broadcast();
+	SettingsChangedEvent.Broadcast();
 }
 
-
-FString UInternationalizationSettingsModel::GetCultureName() const
+FString UInternationalizationSettingsModel::GetEditorCultureName() const
 {
 	FString SavedCultureName;
 	if( !GConfig->GetString( TEXT("Internationalization"), TEXT("Culture"), SavedCultureName, GEditorSettingsIni ) )
@@ -55,13 +53,24 @@ FString UInternationalizationSettingsModel::GetCultureName() const
 	return SavedCultureName;
 }
 
-
-void UInternationalizationSettingsModel::SetCultureName(const FString& CultureName)
+void UInternationalizationSettingsModel::SetEditorCultureName(const FString& CultureName)
 {
 	GConfig->SetString( TEXT("Internationalization"), TEXT("Culture"), *CultureName, GEditorSettingsIni );
-	SettingChangedEvent.Broadcast();
+	SettingsChangedEvent.Broadcast();
+}
+
+FString UInternationalizationSettingsModel::GetNativeGameCultureName() const
+{
+	FString SavedCultureName;
+	GConfig->GetString( TEXT("Internationalization"), TEXT("NativeGameCulture"), SavedCultureName, GEditorSettingsIni );
+	return SavedCultureName;
 }
 
+void UInternationalizationSettingsModel::SetNativeGameCultureName(const FString& CultureName)
+{
+	GConfig->SetString( TEXT("Internationalization"), TEXT("NativeGameCulture"), *CultureName, GEditorSettingsIni );
+	SettingsChangedEvent.Broadcast();
+}
 
 bool UInternationalizationSettingsModel::ShouldLoadLocalizedPropertyNames() const
 {
@@ -73,11 +82,10 @@ bool UInternationalizationSettingsModel::ShouldLoadLocalizedPropertyNames() cons
 	return bShouldLoadLocalizedPropertyNames;
 }
 
-
 void UInternationalizationSettingsModel::ShouldLoadLocalizedPropertyNames(const bool Value)
 {
 	GConfig->SetBool( TEXT("Internationalization"), TEXT("ShouldLoadLocalizedPropertyNames"), Value, GEditorSettingsIni );
-	SettingChangedEvent.Broadcast();
+	SettingsChangedEvent.Broadcast();
 }
 
 bool UInternationalizationSettingsModel::ShouldShowNodesAndPinsUnlocalized() const
@@ -93,4 +101,5 @@ bool UInternationalizationSettingsModel::ShouldShowNodesAndPinsUnlocalized() con
 void UInternationalizationSettingsModel::ShouldShowNodesAndPinsUnlocalized(const bool Value)
 {
 	GConfig->SetBool( TEXT("Internationalization"), TEXT("ShowNodesAndPinsUnlocalized"), Value, GEditorSettingsIni );
+	SettingsChangedEvent.Broadcast();
 }
diff --git a/Engine/Source/Developer/InternationalizationSettings/Private/InternationalizationSettingsModelDetails.cpp b/Engine/Source/Developer/InternationalizationSettings/Private/InternationalizationSettingsModelDetails.cpp
--- a/Engine/Source/Developer/InternationalizationSettings/Private/InternationalizationSettingsModelDetails.cpp
+++ b/Engine/Source/Developer/InternationalizationSettings/Private/InternationalizationSettingsModelDetails.cpp
@@ -10,15 +10,24 @@ struct FCompareCultureByNativeLanguage
 {
 	static FText GetCultureNativeLanguageText( const FCulturePtr Culture )
 	{
-		check( Culture.IsValid() );
-		const FString Language = Culture->GetNativeLanguage();
-		return FText::FromString(Language);
+		return Culture.IsValid() ? FText::FromString(Culture->GetNativeLanguage()) : LOCTEXT("None", "(None)");
 	}
 
 	FORCEINLINE bool operator()( const FCulturePtr A, const FCulturePtr B ) const
 	{
-		check( A.IsValid() );
-		check( B.IsValid() );
+		const FText AText = GetCultureNativeLanguageText(A);
+		const FText BText = GetCultureNativeLanguageText(B);
+
+		// (None) should appear before all else.
+		if(AText.IdenticalTo(LOCTEXT("None", "(None)")))
+		{
+			return true;
+		}
+		// (None) should appear before all else.
+		if(BText.IdenticalTo(LOCTEXT("None", "(None)")))
+		{
+			return false;
+		}
 		return( GetCultureNativeLanguageText( A ).CompareToCaseIgnored( GetCultureNativeLanguageText( B ) ) ) < 0;
 	}
 };
@@ -29,30 +38,52 @@ struct FCompareCultureByNativeRegion
 {
 	static FText GetCultureNativeRegionText( const FCulturePtr Culture )
 	{
-		check( Culture.IsValid() );
-		FString Region = Culture->GetNativeRegion();
-		if ( Region.IsEmpty() )
+		FString Region;
+		if (Culture.IsValid())
+		{
+			Region = Culture->GetNativeRegion();
+			if ( Region.IsEmpty() )
+			{
+				// Fallback to displaying the language, if no region is available
+				return LOCTEXT("NoSpecificRegionOption", "Non-Specific Region");
+			}
+			else
+			{
+				return FText::FromString(Region);
+			}
+		}
+		else
 		{
-			// Fallback to displaying the language, if no region is available
-			return LOCTEXT("NoSpecificRegionOption", "Non-Specific Region");
+			return LOCTEXT("None", "(None)");
 		}
-		return FText::FromString(Region);
 	}
 
 	FORCEINLINE bool operator()( const FCulturePtr A, const FCulturePtr B ) const
 	{
-		check( A.IsValid() );
-		check( B.IsValid() );
+		const FText AText = GetCultureNativeRegionText(A);
+		const FText BText = GetCultureNativeRegionText(B);
+
+		// (None) should appear before all else.
+		if(AText.IdenticalTo(LOCTEXT("None", "(None)")))
+		{
+			return true;
+		}
+		// (None) should appear before all else.
+		if(BText.IdenticalTo(LOCTEXT("None", "(None)")))
+		{
+			return false;
+		}
 		// Non-Specific Region should appear before all else.
-		if(A->GetNativeRegion().IsEmpty())
+		if(AText.IdenticalTo(LOCTEXT("NoSpecificRegionOption", "Non-Specific Region")))
 		{
 			return true;
 		}
 		// Non-Specific Region should appear before all else.
-		if(B->GetNativeRegion().IsEmpty())
+		if(BText.IdenticalTo(LOCTEXT("NoSpecificRegionOption", "Non-Specific Region")))
 		{
 			return false;
 		}
+
 		// Compare native region strings.
 		return( GetCultureNativeRegionText( A ).CompareToCaseIgnored( GetCultureNativeRegionText( B ) ) ) < 0;
 	}
@@ -67,29 +98,50 @@ TSharedRef<IDetailCustomization> FInternationalizationSettingsModelDetails::Make
 FInternationalizationSettingsModelDetails::~FInternationalizationSettingsModelDetails()
 {
 	check(Model.IsValid());
-	Model->OnSettingChanged().RemoveAll(this);
+	Model->OnSettingsChanged().RemoveAll(this);
 }
 
 void FInternationalizationSettingsModelDetails::OnSettingsChanged()
 {
-	FInternationalization& I18N = FInternationalization::Get();
+	// If we made the changes, there's no need to update ourselves from the model.
+	if (IsMakingChangesToModel)
+	{
+		return;
+	}
 
+	FInternationalization& I18N = FInternationalization::Get();
 	check(Model.IsValid());
-	FString SavedCultureName = Model->GetCultureName();
-	if ( !SavedCultureName.IsEmpty() && SavedCultureName != I18N.GetCurrentCulture()->GetName() )
+
+	RequiresRestart = true;
+
+	const FString SavedEditorCultureName = Model->GetEditorCultureName();
+	SelectedEditorCulture = I18N.GetCulture(SavedEditorCultureName);
+	SelectedEditorLanguage = SelectedEditorCulture.IsValid() ? I18N.GetCulture(SelectedEditorCulture->GetTwoLetterISOLanguageName()) : nullptr;
+
+	RefreshAvailableEditorLanguages();
+	EditorLanguageComboBox->SetSelectedItem(SelectedEditorLanguage);
+	RefreshAvailableEditorRegions();
+	EditorRegionComboBox->SetSelectedItem(SelectedEditorCulture);
+
+	const FString SavedNativeGameCultureName = Model->GetNativeGameCultureName();
+	if (!SavedNativeGameCultureName.IsEmpty())
 	{
-		SelectedCulture = I18N.GetCulture(SavedCultureName);
-		RequiresRestart = true;
+		SelectedNativeGameCulture = I18N.GetCulture(SavedNativeGameCultureName);
+		SelectedNativeGameLanguage = SelectedNativeGameCulture.IsValid() ? I18N.GetCulture(SelectedNativeGameCulture->GetTwoLetterISOLanguageName()) : nullptr;
 	}
 	else
 	{
-		SelectedCulture = I18N.GetCurrentCulture();
-		RequiresRestart = false;
+		SelectedNativeGameCulture = nullptr;
+		SelectedNativeGameLanguage = nullptr;
 	}
-	RefreshAvailableLanguages();
-	LanguageComboBox->SetSelectedItem(SelectedLanguage);
-	RefreshAvailableRegions();
-	RegionComboBox->SetSelectedItem(SelectedCulture);
+
+	RefreshAvailableNativeGameLanguages();
+	NativeGameLanguageComboBox->SetSelectedItem(SelectedNativeGameLanguage);
+	RefreshAvailableNativeGameRegions();
+	NativeGameRegionComboBox->SetSelectedItem(SelectedNativeGameCulture);
+
+	LocalizedPropertyNamesCheckBox->SetIsChecked(Model->ShouldLoadLocalizedPropertyNames() ? ECheckBoxState::Checked : ECheckBoxState::Unchecked);
+	UnlocalizedNodesAndPinsCheckBox->SetIsChecked(Model->ShouldShowNodesAndPinsUnlocalized() ? ECheckBoxState::Unchecked : ECheckBoxState::Checked);
 }
 
 void FInternationalizationSettingsModelDetails::CustomizeDetails( IDetailLayoutBuilder& DetailBuilder )
@@ -106,95 +158,176 @@ void FInternationalizationSettingsModelDetails::CustomizeDetails( IDetailLayoutB
 	}
 	check(Model.IsValid());
 
-	Model->OnSettingChanged().AddRaw(this, &FInternationalizationSettingsModelDetails::OnSettingsChanged);
+	Model->OnSettingsChanged().AddRaw(this, &FInternationalizationSettingsModelDetails::OnSettingsChanged);
 
 	// If the saved culture is not the same as the actual current culture, a restart is needed to sync them fully and properly.
-	FString SavedCultureName = Model->GetCultureName();
-	if ( !SavedCultureName.IsEmpty() && SavedCultureName != I18N.GetCurrentCulture()->GetName() )
+	FString SavedEditorCultureName = Model->GetEditorCultureName();
+	if ( !SavedEditorCultureName.IsEmpty() && SavedEditorCultureName != I18N.GetCurrentCulture()->GetName() )
 	{
-		SelectedCulture = I18N.GetCulture(SavedCultureName);
+		SelectedEditorCulture = I18N.GetCulture(SavedEditorCultureName);
 		RequiresRestart = true;
 	}
 	else
 	{
-		SelectedCulture = I18N.GetCurrentCulture();
+		SelectedEditorCulture = I18N.GetCurrentCulture();
 		RequiresRestart = false;
 	}
 
+	const FString SavedNativeGameCultureName = Model->GetNativeGameCultureName();
+	if (!SavedNativeGameCultureName.IsEmpty())
+	{
+		SelectedNativeGameCulture = I18N.GetCulture(SavedNativeGameCultureName);
+		SelectedNativeGameLanguage = SelectedNativeGameCulture.IsValid() ? I18N.GetCulture(SelectedNativeGameCulture->GetTwoLetterISOLanguageName()) : nullptr;
+	}
+	else
+	{
+		SelectedNativeGameCulture = nullptr;
+		SelectedNativeGameLanguage = nullptr;
+	}
+
 	// Populate all our cultures
-	RefreshAvailableCultures();
+	RefreshAvailableEditorCultures();
+	RefreshAvailableNativeGameCultures();
 
 	IDetailCategoryBuilder& CategoryBuilder = DetailBuilder.EditCategory("Internationalization");
 
-	const FText LanguageToolTipText = LOCTEXT("EditorLanguageTooltip", "Change the Editor language (requires restart to take effect)");
-
-	// For use in the Slate macros below, the type must be typedef'd to compile.
-	typedef FCulturePtr ThreadSafeCulturePtr;
-
-	CategoryBuilder.AddCustomRow(LOCTEXT("EditorLanguageLabel", "Language"))
-	.NameContent()
-	[
-		SNew(SHorizontalBox)
-		+SHorizontalBox::Slot()
-		.Padding( FMargin( 0, 1, 0, 1 ) )
-		.FillWidth(1.0f)
-		[
-			SNew(STextBlock)
-			.Text(LOCTEXT("EditorLanguageLabel", "Language"))
-			.Font(DetailBuilder.GetDetailFont())
-			.ToolTipText(LanguageToolTipText)
-		]
-	]
-	.ValueContent()
-	.MaxDesiredWidth(300.0f)
-	[
-		SAssignNew(LanguageComboBox, SComboBox< ThreadSafeCulturePtr > )
-		.OptionsSource( &AvailableLanguages )
-		.InitiallySelectedItem(SelectedLanguage)
-		.OnGenerateWidget(this, &FInternationalizationSettingsModelDetails::OnLanguageGenerateWidget, &DetailBuilder)
-		.ToolTipText(LanguageToolTipText)
-		.OnSelectionChanged(this, &FInternationalizationSettingsModelDetails::OnLanguageSelectionChanged)
-		.Content()
-		[
-			SNew(STextBlock)
-			.Text(this, &FInternationalizationSettingsModelDetails::GetCurrentLanguageText)
-			.Font(DetailBuilder.GetDetailFont())
-		]
-	];
-
-	const FText RegionToolTipText = LOCTEXT("EditorRegionTooltip", "Change the Editor region (requires restart to take effect)");
+	{
+		const FText LanguageToolTipText = LOCTEXT("EditorLanguageTooltip", "Change which language's translations the editor uses. (Requires restart to take effect.)");
+
+		CategoryBuilder.AddCustomRow(LOCTEXT("EditorLanguageLabel", "Editor Localization Language"))
+			.NameContent()
+			[
+				SNew(SHorizontalBox)
+				+SHorizontalBox::Slot()
+				.Padding( FMargin( 0, 1, 0, 1 ) )
+				.FillWidth(1.0f)
+				[
+					SNew(STextBlock)
+					.Text(LOCTEXT("EditorLanguageLabel", "Editor Localization Language"))
+					.Font(DetailBuilder.GetDetailFont())
+					.ToolTipText(LanguageToolTipText)
+				]
+			]
+		.ValueContent()
+			.MaxDesiredWidth(300.0f)
+			[
+				SAssignNew(EditorLanguageComboBox, SComboBox< FCulturePtr > )
+				.OptionsSource( &AvailableEditorLanguages )
+				.InitiallySelectedItem(SelectedEditorLanguage)
+				.OnGenerateWidget(this, &FInternationalizationSettingsModelDetails::OnLanguageGenerateWidget, &DetailBuilder)
+				.ToolTipText(LanguageToolTipText)
+				.OnSelectionChanged(this, &FInternationalizationSettingsModelDetails::OnEditorLanguageSelectionChanged)
+				.Content()
+				[
+					SNew(STextBlock)
+					.Text(this, &FInternationalizationSettingsModelDetails::GetEditorCurrentLanguageText)
+					.Font(DetailBuilder.GetDetailFont())
+				]
+			];
+
+		const FText RegionToolTipText = LOCTEXT("EditorRegionTooltip", "Change which region's translations the editor uses. (Requires restart to take effect.)");
+
+		CategoryBuilder.AddCustomRow(LOCTEXT("EditorRegionLabel", "Editor Localization Region"))
+			.NameContent()
+			[
+				SNew(SHorizontalBox)
+				+SHorizontalBox::Slot()
+				.Padding( FMargin( 0, 1, 0, 1 ) )
+				.FillWidth(1.0f)
+				[
+					SNew(STextBlock)
+					.Text(LOCTEXT("EditorRegionLabel", "Editor Localization Region"))
+					.Font(DetailBuilder.GetDetailFont())
+					.ToolTipText(RegionToolTipText)
+				]
+			]
+		.ValueContent()
+			.MaxDesiredWidth(300.0f)
+			[
+				SAssignNew(EditorRegionComboBox, SComboBox< FCulturePtr > )
+				.OptionsSource( &AvailableEditorRegions )
+				.InitiallySelectedItem(SelectedEditorCulture)
+				.OnGenerateWidget(this, &FInternationalizationSettingsModelDetails::OnRegionGenerateWidget, &DetailBuilder)
+				.ToolTipText(RegionToolTipText)
+				.OnSelectionChanged(this, &FInternationalizationSettingsModelDetails::OnEditorRegionSelectionChanged)
+				.IsEnabled(this, &FInternationalizationSettingsModelDetails::IsEditorRegionSelectionAllowed)
+				.Content()
+				[
+					SNew(STextBlock)
+					.Text(this, &FInternationalizationSettingsModelDetails::GetCurrentEditorRegionText)
+					.Font(DetailBuilder.GetDetailFont())
+				]
+			];
+	}
 
-	CategoryBuilder.AddCustomRow(LOCTEXT("EditorRegionLabel", "Region"))
-	.NameContent()
-	[
-		SNew(SHorizontalBox)
-		+SHorizontalBox::Slot()
-		.Padding( FMargin( 0, 1, 0, 1 ) )
-		.FillWidth(1.0f)
-		[
-			SNew(STextBlock)
-			.Text(LOCTEXT("EditorRegionLabel", "Region"))
-			.Font(DetailBuilder.GetDetailFont())
-			.ToolTipText(RegionToolTipText)
-		]
-	]
-	.ValueContent()
-	.MaxDesiredWidth(300.0f)
-	[
-		SAssignNew(RegionComboBox, SComboBox< ThreadSafeCulturePtr > )
-		.OptionsSource( &AvailableRegions )
-		.InitiallySelectedItem(SelectedCulture)
-		.OnGenerateWidget(this, &FInternationalizationSettingsModelDetails::OnRegionGenerateWidget, &DetailBuilder)
-		.ToolTipText(RegionToolTipText)
-		.OnSelectionChanged(this, &FInternationalizationSettingsModelDetails::OnRegionSelectionChanged)
-		.IsEnabled(this, &FInternationalizationSettingsModelDetails::IsRegionSelectionAllowed)
-		.Content()
-		[
-			SNew(STextBlock)
-			.Text(this, &FInternationalizationSettingsModelDetails::GetCurrentRegionText)
-			.Font(DetailBuilder.GetDetailFont())
-		]
-	];
+	{
+		const FText LanguageToolTipText = LOCTEXT("GameLanguageTooltip", "Change which language the editor treats as native for game localizations. (Requires restart to take effect.)");
+
+		CategoryBuilder.AddCustomRow(LOCTEXT("GameLanguageLabel", "Game Localization Language"))
+			.NameContent()
+			[
+				SNew(SHorizontalBox)
+				+SHorizontalBox::Slot()
+				.Padding( FMargin( 0, 1, 0, 1 ) )
+				.FillWidth(1.0f)
+				[
+					SNew(STextBlock)
+					.Text(LOCTEXT("GameLanguageLabel", "Game Localization Language"))
+					.Font(DetailBuilder.GetDetailFont())
+					.ToolTipText(LanguageToolTipText)
+				]
+			]
+		.ValueContent()
+			.MaxDesiredWidth(300.0f)
+			[
+				SAssignNew(NativeGameLanguageComboBox, SComboBox< FCulturePtr > )
+				.OptionsSource( &AvailableNativeGameLanguages )
+				.InitiallySelectedItem(SelectedNativeGameLanguage)
+				.OnGenerateWidget(this, &FInternationalizationSettingsModelDetails::OnLanguageGenerateWidget, &DetailBuilder)
+				.ToolTipText(LanguageToolTipText)
+				.OnSelectionChanged(this, &FInternationalizationSettingsModelDetails::OnNativeGameLanguageSelectionChanged)
+				.Content()
+				[
+					SNew(STextBlock)
+					.Text(this, &FInternationalizationSettingsModelDetails::GetNativeGameCurrentLanguageText)
+					.Font(DetailBuilder.GetDetailFont())
+				]
+			];
+
+		const FText RegionToolTipText = LOCTEXT("GameRegionTooltip", "Change which region the editor treats as native for game localizations. (Requires restart to take effect.)");
+
+		CategoryBuilder.AddCustomRow(LOCTEXT("GameRegionLabel", "Game Localization Region"))
+			.NameContent()
+			[
+				SNew(SHorizontalBox)
+				+SHorizontalBox::Slot()
+				.Padding( FMargin( 0, 1, 0, 1 ) )
+				.FillWidth(1.0f)
+				[
+					SNew(STextBlock)
+					.Text(LOCTEXT("GameRegionLabel", "Game Localization Region"))
+					.Font(DetailBuilder.GetDetailFont())
+					.ToolTipText(RegionToolTipText)
+				]
+			]
+		.ValueContent()
+			.MaxDesiredWidth(300.0f)
+			[
+				SAssignNew(NativeGameRegionComboBox, SComboBox< FCulturePtr > )
+				.OptionsSource( &AvailableNativeGameRegions )
+				.InitiallySelectedItem(SelectedNativeGameCulture)
+				.OnGenerateWidget(this, &FInternationalizationSettingsModelDetails::OnRegionGenerateWidget, &DetailBuilder)
+				.ToolTipText(RegionToolTipText)
+				.OnSelectionChanged(this, &FInternationalizationSettingsModelDetails::OnNativeGameRegionSelectionChanged)
+				.IsEnabled(this, &FInternationalizationSettingsModelDetails::IsNativeGameRegionSelectionAllowed)
+				.Content()
+				[
+					SNew(STextBlock)
+					.Text(this, &FInternationalizationSettingsModelDetails::GetCurrentNativeGameRegionText)
+					.Font(DetailBuilder.GetDetailFont())
+				]
+			];
+	}
 
 	const FText FieldNamesToolTipText = LOCTEXT("EditorFieldNamesTooltip", "Toggle showing localized field names (requires restart to take effect)");
 
@@ -215,7 +348,7 @@ void FInternationalizationSettingsModelDetails::CustomizeDetails( IDetailLayoutB
 	.ValueContent()
 	.MaxDesiredWidth(300.0f)
 	[
-		SAssignNew(FieldNamesCheckBox, SCheckBox)
+		SAssignNew(LocalizedPropertyNamesCheckBox, SCheckBox)
 		.IsChecked(Model->ShouldLoadLocalizedPropertyNames() ? ECheckBoxState::Checked : ECheckBoxState::Unchecked)
 		.ToolTipText(FieldNamesToolTipText)
 		.OnCheckStateChanged(this, &FInternationalizationSettingsModelDetails::ShouldLoadLocalizedFieldNamesCheckChanged)
@@ -240,7 +373,7 @@ void FInternationalizationSettingsModelDetails::CustomizeDetails( IDetailLayoutB
 	.ValueContent()
 	.MaxDesiredWidth(300.0f)
 	[
-		SAssignNew(FieldNamesCheckBox, SCheckBox)
+		SAssignNew(UnlocalizedNodesAndPinsCheckBox, SCheckBox)
 		.IsChecked(Model->ShouldShowNodesAndPinsUnlocalized() ? ECheckBoxState::Unchecked : ECheckBoxState::Checked)
 		.ToolTipText(NodeAndPinsNamesToolTipText)
 		.OnCheckStateChanged(this, &FInternationalizationSettingsModelDetails::ShouldShowNodesAndPinsUnlocalized)
@@ -271,73 +404,73 @@ void FInternationalizationSettingsModelDetails::CustomizeDetails( IDetailLayoutB
 	];
 }
 
-void FInternationalizationSettingsModelDetails::RefreshAvailableCultures()
+void FInternationalizationSettingsModelDetails::RefreshAvailableEditorCultures()
 {
-	AvailableCultures.Empty();
-	TArray<FCultureRef> LocalizedCultures;
-	FInternationalization::Get().GetCulturesWithAvailableLocalization(FPaths::GetEditorLocalizationPaths(), LocalizedCultures, true);
-	for (const FCultureRef& Culture : LocalizedCultures)
 	{
-		AvailableCultures.Add(Culture);
-	}
+		AvailableEditorCultures.Empty();
+		TArray<FCultureRef> LocalizedCultures;
+		FInternationalization::Get().GetCulturesWithAvailableLocalization(FPaths::GetEditorLocalizationPaths(), LocalizedCultures, true);
+		for (const FCultureRef& Culture : LocalizedCultures)
+		{
+			AvailableEditorCultures.Add(Culture);
+		}
 
-	// Update our selected culture based on the available choices
-	if ( !AvailableCultures.Contains( SelectedCulture ) )
-	{
-		SelectedCulture = NULL;
+		// Update our selected culture based on the available choices
+		if ( !AvailableEditorCultures.Contains( SelectedEditorCulture ) )
+		{
+			SelectedEditorCulture = NULL;
+		}
+		RefreshAvailableEditorLanguages();
+		RefreshAvailableEditorRegions();
 	}
-
-	RefreshAvailableLanguages();
 }
 
-void FInternationalizationSettingsModelDetails::RefreshAvailableLanguages()
+void FInternationalizationSettingsModelDetails::RefreshAvailableEditorLanguages()
 {
-	AvailableLanguages.Empty();
+	AvailableEditorLanguages.Empty();
 	
 	FString SelectedLanguageName;
-	if ( SelectedCulture.IsValid() )
+	if ( SelectedEditorCulture.IsValid() )
 	{
-		SelectedLanguageName = SelectedCulture->GetNativeLanguage();
+		SelectedLanguageName = SelectedEditorCulture->GetNativeLanguage();
 	}
 
 	// Setup the language list
-	for( const auto& Culture : AvailableCultures )
+	for( const auto& Culture : AvailableEditorCultures )
 	{
 		FCulturePtr LanguageCulture = FInternationalization::Get().GetCulture(Culture->GetTwoLetterISOLanguageName());
-		if (LanguageCulture.IsValid() && AvailableLanguages.Find(LanguageCulture) == INDEX_NONE)
+		if (LanguageCulture.IsValid() && AvailableEditorLanguages.Find(LanguageCulture) == INDEX_NONE)
 		{
-			AvailableLanguages.Add(LanguageCulture);
+			AvailableEditorLanguages.Add(LanguageCulture);
 				
 			// Do we have a match for the base language
 			const FString CultureLanguageName = Culture->GetNativeLanguage();
 			if ( SelectedLanguageName == CultureLanguageName)
 			{
-				SelectedLanguage = Culture;
+				SelectedEditorLanguage = Culture;
 			}
 		}
 	}
 
-	AvailableLanguages.Sort( FCompareCultureByNativeLanguage() );
-
-	RefreshAvailableRegions();
+	AvailableEditorLanguages.Sort( FCompareCultureByNativeLanguage() );
 }
 
-void FInternationalizationSettingsModelDetails::RefreshAvailableRegions()
+void FInternationalizationSettingsModelDetails::RefreshAvailableEditorRegions()
 {
-	AvailableRegions.Empty();
+	AvailableEditorRegions.Empty();
 
 	FCulturePtr DefaultCulture;
-	if ( SelectedLanguage.IsValid() )
+	if ( SelectedEditorLanguage.IsValid() )
 	{
-		const FString SelectedLanguageName = SelectedLanguage->GetTwoLetterISOLanguageName();
+		const FString SelectedLanguageName = SelectedEditorLanguage->GetTwoLetterISOLanguageName();
 
 		// Setup the region list
-		for( const auto& Culture : AvailableCultures )
+		for( const auto& Culture : AvailableEditorCultures )
 		{
 			const FString CultureLanguageName = Culture->GetTwoLetterISOLanguageName();
 			if ( SelectedLanguageName == CultureLanguageName)
 			{
-				AvailableRegions.Add( Culture );
+				AvailableEditorRegions.Add( Culture );
 
 				// If this doesn't have a valid region... assume it's the default
 				const FString CultureRegionName = Culture->GetNativeRegion();
@@ -348,119 +481,248 @@ void FInternationalizationSettingsModelDetails::RefreshAvailableRegions()
 			}
 		}
 
-		AvailableRegions.Sort( FCompareCultureByNativeRegion() );
+		AvailableEditorRegions.Sort( FCompareCultureByNativeRegion() );
 	}
 
 	// If we have a preferred default (or there's only one in the list), select that now
 	if ( !DefaultCulture.IsValid() )
 	{
-		if(AvailableRegions.Num() == 1)
+		if(AvailableEditorRegions.Num() == 1)
 		{
-			DefaultCulture = AvailableRegions.Last();
+			DefaultCulture = AvailableEditorRegions.Last();
 		}
 	}
 
 	if ( DefaultCulture.IsValid() )
 	{
 		// Set it as our default region, if one hasn't already been chosen
-		if ( !SelectedCulture.IsValid() && RegionComboBox.IsValid() )
+		if ( !SelectedEditorCulture.IsValid() && EditorRegionComboBox.IsValid() )
 		{
 			// We have to update the combo box like this, otherwise it'll do a null selection when we next click on it
-			RegionComboBox->SetSelectedItem( DefaultCulture );
+			EditorRegionComboBox->SetSelectedItem( DefaultCulture );
 		}
 	}
 
-	if ( RegionComboBox.IsValid() )
+	if ( EditorRegionComboBox.IsValid() )
 	{
-		RegionComboBox->RefreshOptions();
+		EditorRegionComboBox->RefreshOptions();
 	}
 }
 
-FText FInternationalizationSettingsModelDetails::GetCurrentLanguageText() const
+FText FInternationalizationSettingsModelDetails::GetEditorCurrentLanguageText() const
 {
-	if( SelectedLanguage.IsValid() )
+	if( SelectedEditorLanguage.IsValid() )
 	{
-		return FCompareCultureByNativeLanguage::GetCultureNativeLanguageText(SelectedLanguage);
+		return FCompareCultureByNativeLanguage::GetCultureNativeLanguageText(SelectedEditorLanguage);
 	}
 	return FText::GetEmpty();
 }
 
-TSharedRef<SWidget> FInternationalizationSettingsModelDetails::OnLanguageGenerateWidget( FCulturePtr Culture, IDetailLayoutBuilder* DetailBuilder ) const
-{
-	return SNew(STextBlock)
-		.Text(FCompareCultureByNativeLanguage::GetCultureNativeLanguageText(Culture))
-		.Font(DetailBuilder->GetDetailFont());
-}
-
-void FInternationalizationSettingsModelDetails::OnLanguageSelectionChanged( FCulturePtr Culture, ESelectInfo::Type SelectionType )
+void FInternationalizationSettingsModelDetails::OnEditorLanguageSelectionChanged( FCulturePtr Culture, ESelectInfo::Type SelectionType )
 {
-	SelectedLanguage = Culture;
-	SelectedCulture = NULL;
-	RegionComboBox->ClearSelection();
-
-	RefreshAvailableRegions();
+	SelectedEditorLanguage = Culture;
+	SelectedEditorCulture = Culture;
+	Model->SetEditorCultureName(SelectedEditorCulture->GetName());
+	RequiresRestart = true;
+	RefreshAvailableEditorRegions();
 }
 
-FText FInternationalizationSettingsModelDetails::GetCurrentRegionText() const
+FText FInternationalizationSettingsModelDetails::GetCurrentEditorRegionText() const
 {
-	if( SelectedCulture.IsValid() )
+	if( SelectedEditorCulture.IsValid() )
 	{
-		return FCompareCultureByNativeRegion::GetCultureNativeRegionText(SelectedCulture);
+		return FCompareCultureByNativeRegion::GetCultureNativeRegionText(SelectedEditorCulture);
 	}
 	return FText::GetEmpty();
 }
 
-TSharedRef<SWidget> FInternationalizationSettingsModelDetails::OnRegionGenerateWidget( FCulturePtr Culture, IDetailLayoutBuilder* DetailBuilder ) const
+void FInternationalizationSettingsModelDetails::OnEditorRegionSelectionChanged( FCulturePtr Culture, ESelectInfo::Type SelectionType )
 {
-	return SNew(STextBlock)
-		.Text(FCompareCultureByNativeRegion::GetCultureNativeRegionText(Culture))
-		.Font(DetailBuilder->GetDetailFont());
+	TGuardValue<bool> Guard(IsMakingChangesToModel, true);
+	IsMakingChangesToModel = true;
+	SelectedEditorCulture = Culture;
+	Model->SetEditorCultureName(SelectedEditorCulture->GetName());
+	RequiresRestart = true;
 }
 
-void FInternationalizationSettingsModelDetails::OnRegionSelectionChanged( FCulturePtr Culture, ESelectInfo::Type SelectionType )
+bool FInternationalizationSettingsModelDetails::IsEditorRegionSelectionAllowed() const
 {
-	SelectedCulture = Culture;
-
-	HandleShutdownPostPackagesSaved();
+	return SelectedEditorLanguage.IsValid();
 }
 
-bool FInternationalizationSettingsModelDetails::IsRegionSelectionAllowed() const
+void FInternationalizationSettingsModelDetails::RefreshAvailableNativeGameCultures()
 {
-	return SelectedLanguage.IsValid();
+	{
+		AvailableNativeGameCultures.Empty();
+		TArray<FCultureRef> LocalizedCultures;
+		FInternationalization::Get().GetCulturesWithAvailableLocalization(FPaths::GetGameLocalizationPaths(), LocalizedCultures, true);
+		for (const FCultureRef& Culture : LocalizedCultures)
+		{
+			AvailableNativeGameCultures.Add(Culture);
+		}
+		AvailableNativeGameCultures.Add(nullptr);
+
+		// Update our selected culture based on the available choices
+		if ( !AvailableNativeGameCultures.Contains( SelectedNativeGameCulture ) )
+		{
+			SelectedNativeGameCulture = NULL;
+		}
+		RefreshAvailableNativeGameLanguages();
+		RefreshAvailableNativeGameRegions();
+	}
 }
 
-EVisibility FInternationalizationSettingsModelDetails::GetInternationalizationRestartRowVisibility() const
+void FInternationalizationSettingsModelDetails::RefreshAvailableNativeGameLanguages()
 {
-	return RequiresRestart ? EVisibility::Visible : EVisibility::Collapsed;
+	AvailableNativeGameLanguages.Empty();
+
+	FString SelectedLanguageName;
+	if ( SelectedNativeGameCulture.IsValid() )
+	{
+		SelectedLanguageName = SelectedNativeGameCulture->GetNativeLanguage();
+	}
+
+	// Setup the language list
+	for( const auto& Culture : AvailableNativeGameCultures )
+	{
+		FCulturePtr LanguageCulture = Culture.IsValid() ? FInternationalization::Get().GetCulture(Culture->GetTwoLetterISOLanguageName()) : nullptr;
+		if (AvailableNativeGameLanguages.Find(LanguageCulture) == INDEX_NONE)
+		{
+			AvailableNativeGameLanguages.Add(LanguageCulture);
+
+			// Do we have a match for the base language
+			const FString CultureLanguageName = Culture.IsValid() ? Culture->GetNativeLanguage() : TEXT("");
+			if ( SelectedLanguageName == CultureLanguageName)
+			{
+				SelectedNativeGameLanguage = Culture;
+			}
+		}
+	}
+
+	AvailableNativeGameLanguages.Sort( FCompareCultureByNativeLanguage() );
 }
 
-void FInternationalizationSettingsModelDetails::ShouldLoadLocalizedFieldNamesCheckChanged(ECheckBoxState CheckState)
+void FInternationalizationSettingsModelDetails::RefreshAvailableNativeGameRegions()
 {
-	HandleShutdownPostPackagesSaved();
-}
+	AvailableNativeGameRegions.Empty();
 
+	FCulturePtr DefaultCulture;
+	{
+		const FString SelectedLanguageName = SelectedNativeGameLanguage.IsValid() ? SelectedNativeGameLanguage->GetTwoLetterISOLanguageName() : TEXT("");
 
-void FInternationalizationSettingsModelDetails::HandleShutdownPostPackagesSaved()
-{
-	if ( SelectedCulture.IsValid() )
+		// Setup the region list
+		for( const auto& Culture : AvailableNativeGameCultures )
+		{
+			const FString CultureLanguageName = Culture.IsValid() ? Culture->GetTwoLetterISOLanguageName() : TEXT("");
+			if ( SelectedLanguageName == CultureLanguageName)
+			{
+				AvailableNativeGameRegions.Add( Culture );
+
+				// If this doesn't have a valid region... assume it's the default
+				const FString CultureRegionName = Culture.IsValid() ? Culture->GetNativeRegion() : TEXT("");
+				if ( CultureRegionName.IsEmpty() )
+				{	
+					DefaultCulture = Culture;
+				}
+			}
+		}
+
+		AvailableNativeGameRegions.Sort( FCompareCultureByNativeRegion() );
+	}
+
+	// If we have a preferred default (or there's only one in the list), select that now
+	if ( !DefaultCulture.IsValid() )
 	{
-		check(Model.IsValid());
-		Model->SetCultureName(SelectedCulture->GetName());
-		Model->ShouldLoadLocalizedPropertyNames(FieldNamesCheckBox->IsChecked());
-		if(SelectedCulture != FInternationalization::Get().GetCurrentCulture())
+		if(AvailableNativeGameRegions.Num() == 1)
 		{
-			RequiresRestart = true;
+			DefaultCulture = AvailableNativeGameRegions.Last();
 		}
-		else
+	}
+
+	if ( DefaultCulture.IsValid() )
+	{
+		// Set it as our default region, if one hasn't already been chosen
+		if ( !SelectedNativeGameCulture.IsValid() && NativeGameRegionComboBox.IsValid() )
 		{
-			RequiresRestart = false;
+			// We have to update the combo box like this, otherwise it'll do a null selection when we next click on it
+			NativeGameRegionComboBox->SetSelectedItem( DefaultCulture );
 		}
 	}
+
+	if ( NativeGameRegionComboBox.IsValid() )
+	{
+		NativeGameRegionComboBox->RefreshOptions();
+	}
+}
+
+FText FInternationalizationSettingsModelDetails::GetNativeGameCurrentLanguageText() const
+{
+	if( SelectedNativeGameLanguage.IsValid() )
+	{
+		return FCompareCultureByNativeLanguage::GetCultureNativeLanguageText(SelectedNativeGameLanguage);
+	}
+	return LOCTEXT("None", "(None)");
+}
+
+void FInternationalizationSettingsModelDetails::OnNativeGameLanguageSelectionChanged( FCulturePtr Culture, ESelectInfo::Type SelectionType )
+{
+	TGuardValue<bool> Guard(IsMakingChangesToModel, true);
+	SelectedNativeGameLanguage = Culture;
+	SelectedNativeGameCulture = Culture;
+	Model->SetNativeGameCultureName(SelectedNativeGameCulture.IsValid() ? SelectedNativeGameCulture->GetName() : TEXT(""));
+	RefreshAvailableNativeGameRegions();
+}
+
+FText FInternationalizationSettingsModelDetails::GetCurrentNativeGameRegionText() const
+{
+	if( SelectedNativeGameCulture.IsValid() )
+	{
+		return FCompareCultureByNativeRegion::GetCultureNativeRegionText(SelectedNativeGameCulture);
+	}
+	return LOCTEXT("None", "(None)");
+}
+
+void FInternationalizationSettingsModelDetails::OnNativeGameRegionSelectionChanged( FCulturePtr Culture, ESelectInfo::Type SelectionType )
+{
+	TGuardValue<bool> Guard(IsMakingChangesToModel, true);
+	SelectedNativeGameCulture = Culture;
+	Model->SetNativeGameCultureName(SelectedNativeGameCulture.IsValid() ? SelectedNativeGameCulture->GetName() : TEXT(""));
+	RequiresRestart = true;
+}
+
+bool FInternationalizationSettingsModelDetails::IsNativeGameRegionSelectionAllowed() const
+{
+	return SelectedNativeGameLanguage.IsValid();
+}
+
+TSharedRef<SWidget> FInternationalizationSettingsModelDetails::OnLanguageGenerateWidget( FCulturePtr Culture, IDetailLayoutBuilder* DetailBuilder ) const
+{
+	return SNew(STextBlock)
+		.Text(FCompareCultureByNativeLanguage::GetCultureNativeLanguageText(Culture))
+		.Font(DetailBuilder->GetDetailFont());
+}
+
+TSharedRef<SWidget> FInternationalizationSettingsModelDetails::OnRegionGenerateWidget( FCulturePtr Culture, IDetailLayoutBuilder* DetailBuilder ) const
+{
+	return SNew(STextBlock)
+		.Text(FCompareCultureByNativeRegion::GetCultureNativeRegionText(Culture))
+		.Font(DetailBuilder->GetDetailFont());
+}
+
+EVisibility FInternationalizationSettingsModelDetails::GetInternationalizationRestartRowVisibility() const
+{
+	return RequiresRestart ? EVisibility::Visible : EVisibility::Collapsed;
+}
+
+void FInternationalizationSettingsModelDetails::ShouldLoadLocalizedFieldNamesCheckChanged(ECheckBoxState CheckState)
+{
+	TGuardValue<bool> Guard(IsMakingChangesToModel, true);
+	Model->ShouldLoadLocalizedPropertyNames(LocalizedPropertyNamesCheckBox->IsChecked());
 }
 
 void FInternationalizationSettingsModelDetails::ShouldShowNodesAndPinsUnlocalized(ECheckBoxState CheckState)
 {
-	check(Model.IsValid());
+	TGuardValue<bool> Guard(IsMakingChangesToModel, true);
 	Model->ShouldShowNodesAndPinsUnlocalized(CheckState == ECheckBoxState::Unchecked);
 
 	// Find all Schemas and force a visualization cache clear
diff --git a/Engine/Source/Developer/InternationalizationSettings/Public/InternationalizationSettingsModelDetails.h b/Engine/Source/Developer/InternationalizationSettings/Public/InternationalizationSettingsModelDetails.h
--- a/Engine/Source/Developer/InternationalizationSettings/Public/InternationalizationSettingsModelDetails.h
+++ b/Engine/Source/Developer/InternationalizationSettings/Public/InternationalizationSettingsModelDetails.h
@@ -12,7 +12,8 @@ public:
 
 	/** Default Constructor */
 	FInternationalizationSettingsModelDetails()
-		:	RequiresRestart(false)
+		: RequiresRestart(false)
+		, IsMakingChangesToModel(false)
 	{
 	}
 
@@ -30,34 +31,58 @@ private:
 	void PromptForRestart() const;
 
 	/** Look for available cultures */
-	void RefreshAvailableCultures();
+	void RefreshAvailableEditorCultures();
 
 	/** Look for available languages */
-	void RefreshAvailableLanguages();
+	void RefreshAvailableEditorLanguages();
 
 	/** Look for available regions */
-	void RefreshAvailableRegions();
+	void RefreshAvailableEditorRegions();
 
-	/** Delegate called to display the current language */
-	FText GetCurrentLanguageText() const;
+	/** Delegate called to display the current editor language */
+	FText GetEditorCurrentLanguageText() const;
 
-	/** Generate a widget for the language combo */
-	TSharedRef<SWidget> OnLanguageGenerateWidget( FCulturePtr Culture, IDetailLayoutBuilder* DetailBuilder ) const;
+	/** Delegate called when the editor language selection has changed */
+	void OnEditorLanguageSelectionChanged( FCulturePtr Culture, ESelectInfo::Type SelectionType );
 
-	/** Delegate called when the language selection has changed */
-	void OnLanguageSelectionChanged( FCulturePtr Culture, ESelectInfo::Type SelectionType );
+	/** Delegate called to display the current editor region */
+	FText GetCurrentEditorRegionText() const;
 
-	/** Generate a widget for the region combo */
-	TSharedRef<SWidget> OnRegionGenerateWidget( FCulturePtr Culture, IDetailLayoutBuilder* DetailBuilder ) const;
+	/** Delegate called when the editor region selection has changed */
+	void OnEditorRegionSelectionChanged( FCulturePtr Culture, ESelectInfo::Type SelectionType );
+
+	/** Delegate called to see if we're allowed to change the region */
+	bool IsEditorRegionSelectionAllowed() const;
+
+	/** Look for available cultures */
+	void RefreshAvailableNativeGameCultures();
+
+	/** Look for available languages */
+	void RefreshAvailableNativeGameLanguages();
+
+	/** Look for available regions */
+	void RefreshAvailableNativeGameRegions();
+
+	/** Delegate called to display the current NativeGame language */
+	FText GetNativeGameCurrentLanguageText() const;
+
+	/** Delegate called when the NativeGame language selection has changed */
+	void OnNativeGameLanguageSelectionChanged( FCulturePtr Culture, ESelectInfo::Type SelectionType );
 
-	/** Delegate called to display the current region */
-	FText GetCurrentRegionText() const;
+	/** Delegate called to display the current NativeGame region */
+	FText GetCurrentNativeGameRegionText() const;
 
-	/** Delegate called when the region selection has changed */
-	void OnRegionSelectionChanged( FCulturePtr Culture, ESelectInfo::Type SelectionType );
+	/** Delegate called when the NativeGame region selection has changed */
+	void OnNativeGameRegionSelectionChanged( FCulturePtr Culture, ESelectInfo::Type SelectionType );
 
 	/** Delegate called to see if we're allowed to change the region */
-	bool IsRegionSelectionAllowed() const;
+	bool IsNativeGameRegionSelectionAllowed() const;
+
+	/** Generate a widget for the language combo */
+	TSharedRef<SWidget> OnLanguageGenerateWidget( FCulturePtr Culture, IDetailLayoutBuilder* DetailBuilder ) const;
+
+	/** Generate a widget for the region combo */
+	TSharedRef<SWidget> OnRegionGenerateWidget( FCulturePtr Culture, IDetailLayoutBuilder* DetailBuilder ) const;
 
 	/** Delegate called to determine whether to collapse or display the restart row. */
 	EVisibility GetInternationalizationRestartRowVisibility() const;
@@ -75,28 +100,54 @@ private:
 	TWeakObjectPtr<UInternationalizationSettingsModel> Model;
 
 	/** The culture selected at present */
-	FCulturePtr SelectedCulture;
+	FCulturePtr SelectedEditorCulture;
+
+	/** The language selected at present */
+	FCulturePtr SelectedEditorLanguage;
+
+	/** The cultures we have available to us */
+	TArray< FCulturePtr > AvailableEditorCultures;
+
+	/** The languages we have available to us */
+	TArray< FCulturePtr > AvailableEditorLanguages;
+
+	/** The dropdown for the languages */
+	TSharedPtr<SComboBox< FCulturePtr > > EditorLanguageComboBox;
+
+	/** The regions we have available to us */
+	TArray< FCulturePtr > AvailableEditorRegions;
+
+	/** The dropdown for the regions */
+	TSharedPtr<SComboBox< FCulturePtr > > EditorRegionComboBox;
+
+	/** The culture selected at present */
+	FCulturePtr SelectedNativeGameCulture;
 
 	/** The language selected at present */
-	FCulturePtr SelectedLanguage;
+	FCulturePtr SelectedNativeGameLanguage;
 
 	/** The cultures we have available to us */
-	TArray< FCulturePtr > AvailableCultures;
+	TArray< FCulturePtr > AvailableNativeGameCultures;
 
 	/** The languages we have available to us */
-	TArray< FCulturePtr > AvailableLanguages;
+	TArray< FCulturePtr > AvailableNativeGameLanguages;
 
 	/** The dropdown for the languages */
-	TSharedPtr<SComboBox< FCulturePtr > > LanguageComboBox;
+	TSharedPtr<SComboBox< FCulturePtr > > NativeGameLanguageComboBox;
 
 	/** The regions we have available to us */
-	TArray< FCulturePtr > AvailableRegions;
+	TArray< FCulturePtr > AvailableNativeGameRegions;
 
 	/** The dropdown for the regions */
-	TSharedPtr<SComboBox< FCulturePtr > > RegionComboBox;
+	TSharedPtr<SComboBox< FCulturePtr > > NativeGameRegionComboBox;
 
-	/** The check box for using localized field names */
-	TSharedPtr<SCheckBox> FieldNamesCheckBox;
+	/** The check box for using localized property names */
+	TSharedPtr<SCheckBox> LocalizedPropertyNamesCheckBox;
+
+	/** The check box for showing nodes and pins unlocalized */
+	TSharedPtr<SCheckBox> UnlocalizedNodesAndPinsCheckBox;
 
 	bool RequiresRestart;
+
+	bool IsMakingChangesToModel;
 };
\ No newline at end of file
diff --git a/Engine/Source/Editor/EditorSettingsViewer/Private/EditorSettingsViewerModule.cpp b/Engine/Source/Editor/EditorSettingsViewer/Private/EditorSettingsViewerModule.cpp
--- a/Engine/Source/Editor/EditorSettingsViewer/Private/EditorSettingsViewerModule.cpp
+++ b/Engine/Source/Editor/EditorSettingsViewer/Private/EditorSettingsViewerModule.cpp
@@ -110,7 +110,6 @@ protected:
 			RegionAndLanguageSettingsSection->OnImport().BindRaw(this, &FEditorSettingsViewerModule::HandleRegionAndLanguageImport);
 			RegionAndLanguageSettingsSection->OnSaveDefaults().BindRaw(this, &FEditorSettingsViewerModule::HandleRegionAndLanguageSaveDefaults);
 			RegionAndLanguageSettingsSection->OnResetDefaults().BindRaw(this, &FEditorSettingsViewerModule::HandleRegionAndLanguageResetToDefault);
-			GetMutableDefault<UInternationalizationSettingsModel>()->OnSettingChanged().AddRaw(this, &FEditorSettingsViewerModule::HandleRegionAndLanguageSettingChanged);
 		}
 
 		// input bindings
@@ -371,19 +370,23 @@ private:
 
 	bool HandleRegionAndLanguageExport(const FString& FileName)
 	{
-		FString CultureName = GetMutableDefault<UInternationalizationSettingsModel>()->GetCultureName();
-		GConfig->SetString( TEXT("Internationalization"), TEXT("Culture"), *CultureName, FileName );
-		GConfig->Flush( false, FileName );
-		return true;
+		UInternationalizationSettingsModel* Model = GetMutableDefault<UInternationalizationSettingsModel>();
+		GConfig->Flush(false, Model->GetDefaultConfigFilename());
+		return BackupFile(Model->GetDefaultConfigFilename(), FileName);
 	}
 
 	bool HandleRegionAndLanguageImport(const FString& FileName)
 	{
-		FString CultureName;
-		GConfig->LoadFile(FileName);
-		GConfig->GetString( TEXT("Internationalization"), TEXT("Culture"), CultureName, FileName );
-		GetMutableDefault<UInternationalizationSettingsModel>()->SetCultureName(CultureName);
-		return true;
+		UInternationalizationSettingsModel* Model = GetMutableDefault<UInternationalizationSettingsModel>();
+		if( EAppReturnType::Ok == ShowRestartWarning(LOCTEXT("ImportRegionAndLanguage_Title", "Import Region & Language")))
+		{
+			FUnrealEdMisc::Get().SetConfigRestoreFilename(FileName, Model->GetDefaultConfigFilename());
+			FUnrealEdMisc::Get().RestartEditor(false);
+
+			return true;
+		}
+
+		return false;
 	}
 
 	bool HandleRegionAndLanguageSaveDefaults()
@@ -398,12 +401,6 @@ private:
 		return true;
 	}
 
-	void HandleRegionAndLanguageSettingChanged()
-	{
-		ISettingsEditorModule& SettingsEditorModule = FModuleManager::GetModuleChecked<ISettingsEditorModule>("SettingsEditor");
-		SettingsEditorModule.OnApplicationRestartRequired();
-	}
-
 private:
 
 	/** Holds a pointer to the settings editor's view model. */
diff --git a/Engine/Source/Editor/TranslationEditor/Private/TranslationDataManager.cpp b/Engine/Source/Editor/TranslationEditor/Private/TranslationDataManager.cpp
--- a/Engine/Source/Editor/TranslationEditor/Private/TranslationDataManager.cpp
+++ b/Engine/Source/Editor/TranslationEditor/Private/TranslationDataManager.cpp
@@ -627,6 +627,29 @@ void FTranslationDataManager::LoadFromArchive(TArray<UTranslationUnit*>& InTrans
 	{
 		TSharedRef< FInternationalizationArchive > Archive = ArchivePtr.ToSharedRef();
 
+		for (auto Iter = Archive->GetEntryIterator(); Iter; ++Iter)
+		{
+			const FString SourceText = Iter->Key;
+			const TSharedRef< FArchiveEntry >& ArchiveEntry = Iter->Value;
+
+			const bool HasTranslationUnit = InTranslationUnits.ContainsByPredicate([&](UTranslationUnit* TU) -> bool
+			{
+				return TU->Namespace == ArchiveEntry->Namespace && TU->Source == ArchiveEntry->Source.Text;
+			});
+
+			// If a translation unit doesn't exist for this namespace and source, add it.
+			if (!HasTranslationUnit)
+			{
+				UTranslationUnit* TranslationUnit = NewObject<UTranslationUnit>();
+				check(TranslationUnit != nullptr);
+				TranslationUnit->SetFlags(RF_Transactional);
+				TranslationUnit->HasBeenReviewed = false;
+				TranslationUnit->Source = ArchiveEntry->Source.Text;
+				TranslationUnit->Namespace = ArchiveEntry->Namespace;
+				InTranslationUnits.Add(TranslationUnit);
+			}
+		}
+
 		// Make a local copy of this array before we empty the arrays below (we might have been passed AllTranslations array)
 		TArray<UTranslationUnit*> TranslationUnits;
 		TranslationUnits.Append(InTranslationUnits);
diff --git a/Engine/Source/Runtime/Core/Private/Internationalization/TextLocalizationManager.cpp b/Engine/Source/Runtime/Core/Private/Internationalization/TextLocalizationManager.cpp
--- a/Engine/Source/Runtime/Core/Private/Internationalization/TextLocalizationManager.cpp
+++ b/Engine/Source/Runtime/Core/Private/Internationalization/TextLocalizationManager.cpp
@@ -713,15 +713,16 @@ void FTextLocalizationManager::LoadLocalizationResourcesForCulture(const FString
 	}
 
 	// Collect the localization paths to load from.
-	TArray<FString> LocalizationPaths;
-	if(ShouldLoadGame)
+	TArray<FString> GameLocalizationPaths;
+	if(ShouldLoadGame || GIsEditor)
 	{
-		LocalizationPaths += FPaths::GetGameLocalizationPaths();
+		GameLocalizationPaths += FPaths::GetGameLocalizationPaths();
 	}
+	TArray<FString> EditorLocalizationPaths;
 	if(ShouldLoadEditor)
 	{
-		LocalizationPaths += FPaths::GetEditorLocalizationPaths();
-		LocalizationPaths += FPaths::GetToolTipLocalizationPaths();
+		EditorLocalizationPaths += FPaths::GetEditorLocalizationPaths();
+		EditorLocalizationPaths += FPaths::GetToolTipLocalizationPaths();
 
 		bool bShouldLoadLocalizedPropertyNames = true;
 		if( !GConfig->GetBool( TEXT("Internationalization"), TEXT("ShouldLoadLocalizedPropertyNames"), bShouldLoadLocalizedPropertyNames, GEditorSettingsIni ) )
@@ -730,10 +731,11 @@ void FTextLocalizationManager::LoadLocalizationResourcesForCulture(const FString
 		}
 		if(bShouldLoadLocalizedPropertyNames)
 		{
-			LocalizationPaths += FPaths::GetPropertyNameLocalizationPaths();
+			EditorLocalizationPaths += FPaths::GetPropertyNameLocalizationPaths();
 		}
 	}
-	LocalizationPaths += FPaths::GetEngineLocalizationPaths();
+	TArray<FString> EngineLocalizationPaths;
+	EngineLocalizationPaths += FPaths::GetEngineLocalizationPaths();
 
 	// Prioritized array of localization entry trackers.
 	TArray<FLocalizationEntryTracker> LocalizationEntryTrackers;
@@ -776,17 +778,57 @@ void FTextLocalizationManager::LoadLocalizationResourcesForCulture(const FString
 	};
 
 	TMap< FString, TMap<FString, FString> > LocalizationPathToCultureDirectoryMap;
-	for (const FString& LocalizationPath : LocalizationPaths)
+	for (const FString& LocalizationPath : GameLocalizationPaths)
 	{
 		LocalizationPathToCultureDirectoryMap.Add(LocalizationPath, MapCulturesToDirectories(LocalizationPath));
 	}
+	for (const FString& LocalizationPath : EditorLocalizationPaths)
+	{
+		LocalizationPathToCultureDirectoryMap.Add(LocalizationPath, MapCulturesToDirectories(LocalizationPath));
+	}
+	for (const FString& LocalizationPath : EngineLocalizationPaths)
+	{
+		LocalizationPathToCultureDirectoryMap.Add(LocalizationPath, MapCulturesToDirectories(LocalizationPath));
+	}
+
+	TArray<FString> PrioritizedLocalizationPaths;
+	if (!GIsEditor)
+	{
+		PrioritizedLocalizationPaths += GameLocalizationPaths;
+	}
+	PrioritizedLocalizationPaths += EditorLocalizationPaths;
+	PrioritizedLocalizationPaths += EngineLocalizationPaths;
+
+	// The editor cheats and loads the native culture's localizations.
+	if (GIsEditor)
+	{
+		FString NativeCultureName = TEXT("");
+		GConfig->GetString( TEXT("Internationalization"), TEXT("NativeGameCulture"), NativeCultureName, GEditorSettingsIni );
+
+		if (!NativeCultureName.IsEmpty())
+		{
+			FLocalizationEntryTracker& CultureTracker = LocalizationEntryTrackers[LocalizationEntryTrackers.Add(FLocalizationEntryTracker())];
+			for (const FString& LocalizationPath : GameLocalizationPaths)
+			{
+				const FString* const Entry = LocalizationPathToCultureDirectoryMap[LocalizationPath].Find(FCulture::GetCanonicalName(NativeCultureName));
+				if (Entry)
+				{
+					const FString CulturePath = LocalizationPath / (*Entry);
+
+					CultureTracker.LoadFromDirectory(CulturePath);
+				}
+			}
+			CultureTracker.DetectAndLogConflicts();
+		}
+	}
+
 
 	// Read culture localization resources.
 	const TArray<FString> PrioritizedParentCultureNames = Culture->GetPrioritizedParentCultureNames();
 	for (const FString& ParentCultureName : PrioritizedParentCultureNames)
 	{
 		FLocalizationEntryTracker& CultureTracker = LocalizationEntryTrackers[LocalizationEntryTrackers.Add(FLocalizationEntryTracker())];
-		for (const FString& LocalizationPath : LocalizationPaths)
+		for (const FString& LocalizationPath : PrioritizedLocalizationPaths)
 		{
 			const FString* const Entry = LocalizationPathToCultureDirectoryMap[LocalizationPath].Find(FCulture::GetCanonicalName(ParentCultureName));
 			if (Entry)
