Somewhat cleaned up internationalization settings and fixed a slew of issues with language and region combinations like Chinese for China. Still needs a lot of cleanup, probably a solid rewrite, but this should suffice for now.

diff --git a/Engine/Source/Developer/InternationalizationSettings/Private/InternationalizationSettingsModelDetails.cpp b/Engine/Source/Developer/InternationalizationSettings/Private/InternationalizationSettingsModelDetails.cpp
--- a/Engine/Source/Developer/InternationalizationSettings/Private/InternationalizationSettingsModelDetails.cpp
+++ b/Engine/Source/Developer/InternationalizationSettings/Private/InternationalizationSettingsModelDetails.cpp
@@ -101,47 +101,129 @@ FInternationalizationSettingsModelDetails::~FInternationalizationSettingsModelDe
 	Model->OnSettingsChanged().RemoveAll(this);
 }
 
-void FInternationalizationSettingsModelDetails::OnSettingsChanged()
+void FInternationalizationSettingsModelDetails::UpdateInternalStateFromSettingsModel()
 {
-	// If we made the changes, there's no need to update ourselves from the model.
-	if (IsMakingChangesToModel)
-	{
-		return;
-	}
-
 	FInternationalization& I18N = FInternationalization::Get();
-	check(Model.IsValid());
-
-	RequiresRestart = true;
 
 	const FString SavedEditorCultureName = Model->GetEditorCultureName();
-	SelectedEditorCulture = I18N.GetCulture(SavedEditorCultureName);
-	SelectedEditorLanguage = SelectedEditorCulture.IsValid() ? I18N.GetCulture(SelectedEditorCulture->GetTwoLetterISOLanguageName()) : nullptr;
+	// Attempt to use the culture specified in the settings.
+	if (!SavedEditorCultureName.IsEmpty())
+	{
+		const FCulturePtr SavedEditorCulture = I18N.GetCulture(SavedEditorCultureName);
+		const FCulturePtr SavedEditorLanguage = SavedEditorCulture.IsValid() ? I18N.GetCulture(SavedEditorCulture->GetTwoLetterISOLanguageName()) : nullptr;
 
-	RefreshAvailableEditorLanguages();
-	EditorLanguageComboBox->SetSelectedItem(SelectedEditorLanguage);
-	RefreshAvailableEditorRegions();
-	EditorRegionComboBox->SetSelectedItem(SelectedEditorCulture);
+		// Accept the saved culture's language if it is available.
+		if (AvailableEditorLanguages.Contains(SavedEditorLanguage))
+		{
+			SelectedEditorLanguage = SavedEditorLanguage;
+			RefreshAvailableEditorRegions();
+
+			// Accept the saved culture's region if it is available.
+			if (AvailableEditorRegions.Contains(SavedEditorCulture))
+			{
+				SelectedEditorCulture = SavedEditorCulture;
+			}
+			// Fallback on the first available region of the selected language.
+			else
+			{
+				SelectedEditorCulture = AvailableEditorRegions[0];
+			}
+		}
+		else
+		{
+			SelectedEditorCulture = I18N.GetCurrentCulture();
+			SelectedEditorLanguage = SelectedEditorCulture.IsValid() ? I18N.GetCulture(SelectedEditorCulture->GetTwoLetterISOLanguageName()) : nullptr;
+			RefreshAvailableEditorRegions();
+		}
+	}
+	// Fallback on the current culture of the editor.
+	else
+	{
+		SelectedEditorCulture = I18N.GetCurrentCulture();
+		SelectedEditorLanguage = SelectedEditorCulture.IsValid() ? I18N.GetCulture(SelectedEditorCulture->GetTwoLetterISOLanguageName()) : nullptr;
+		RefreshAvailableEditorRegions();
+	}
+	check(AvailableEditorLanguages.Contains(SelectedEditorLanguage) && AvailableEditorRegions.Contains(SelectedEditorCulture));
+	if (EditorLanguageComboBox.IsValid())
+	{
+		EditorLanguageComboBox->RefreshOptions();
+		EditorLanguageComboBox->SetSelectedItem(SelectedEditorLanguage);
+	}
+	if (EditorRegionComboBox.IsValid())
+	{
+		EditorRegionComboBox->RefreshOptions();
+		EditorRegionComboBox->SetSelectedItem(SelectedEditorCulture);
+	}
 
 	const FString SavedNativeGameCultureName = Model->GetNativeGameCultureName();
+	// Attempt to use the culture specified in the settings.
 	if (!SavedNativeGameCultureName.IsEmpty())
 	{
-		SelectedNativeGameCulture = I18N.GetCulture(SavedNativeGameCultureName);
-		SelectedNativeGameLanguage = SelectedNativeGameCulture.IsValid() ? I18N.GetCulture(SelectedNativeGameCulture->GetTwoLetterISOLanguageName()) : nullptr;
+		const FCulturePtr SavedNativeGameCulture = I18N.GetCulture(SavedNativeGameCultureName);
+		const FCulturePtr SavedNativeGameLanguage = SavedNativeGameCulture.IsValid() ? I18N.GetCulture(SavedNativeGameCulture->GetTwoLetterISOLanguageName()) : nullptr;
+
+		// Accept the saved culture's language if it is available.
+		if (AvailableNativeGameLanguages.Contains(SavedNativeGameLanguage))
+		{
+			SelectedNativeGameLanguage = SavedNativeGameLanguage;
+			RefreshAvailableNativeGameRegions();
+
+			// Accept the saved culture's region if it is available.
+			if (AvailableNativeGameRegions.Contains(SavedNativeGameCulture))
+			{
+				SelectedNativeGameCulture = SavedNativeGameCulture;
+			}
+			// Fallback on the first available region of the selected language.
+			else
+			{
+				SelectedNativeGameCulture = AvailableNativeGameRegions[0];
+			}
+		}
+		else
+		{
+			SelectedNativeGameCulture = I18N.GetCurrentCulture();
+			SelectedNativeGameLanguage = SelectedNativeGameCulture.IsValid() ? I18N.GetCulture(SelectedNativeGameCulture->GetTwoLetterISOLanguageName()) : nullptr;
+			RefreshAvailableNativeGameRegions();
+		}
 	}
+	// Fallback to no specified language or region.
 	else
 	{
 		SelectedNativeGameCulture = nullptr;
 		SelectedNativeGameLanguage = nullptr;
+		RefreshAvailableNativeGameRegions();
+	}
+	check(AvailableNativeGameLanguages.Contains(SelectedNativeGameLanguage) && AvailableNativeGameRegions.Contains(SelectedNativeGameCulture));
+	if (NativeGameLanguageComboBox.IsValid())
+	{
+		NativeGameLanguageComboBox->RefreshOptions();
+		NativeGameLanguageComboBox->SetSelectedItem(SelectedNativeGameLanguage);
+	}
+	if (NativeGameRegionComboBox.IsValid())
+	{
+		NativeGameRegionComboBox->RefreshOptions();
+		NativeGameRegionComboBox->SetSelectedItem(SelectedNativeGameCulture);
 	}
 
-	RefreshAvailableNativeGameLanguages();
-	NativeGameLanguageComboBox->SetSelectedItem(SelectedNativeGameLanguage);
-	RefreshAvailableNativeGameRegions();
-	NativeGameRegionComboBox->SetSelectedItem(SelectedNativeGameCulture);
+	if (LocalizedPropertyNamesCheckBox.IsValid())
+	{
+		LocalizedPropertyNamesCheckBox->SetIsChecked(Model->ShouldLoadLocalizedPropertyNames() ? ECheckBoxState::Checked : ECheckBoxState::Unchecked);
+	}
+	if (UnlocalizedNodesAndPinsCheckBox.IsValid())
+	{
+		UnlocalizedNodesAndPinsCheckBox->SetIsChecked(Model->ShouldShowNodesAndPinsUnlocalized() ? ECheckBoxState::Unchecked : ECheckBoxState::Checked);
+	}
+}
+
+void FInternationalizationSettingsModelDetails::OnSettingsChanged()
+{
+	// If we made the changes, there's no need to update ourselves from the model.
+	if (IsMakingChangesToModel)
+	{
+		return;
+	}
 
-	LocalizedPropertyNamesCheckBox->SetIsChecked(Model->ShouldLoadLocalizedPropertyNames() ? ECheckBoxState::Checked : ECheckBoxState::Unchecked);
-	UnlocalizedNodesAndPinsCheckBox->SetIsChecked(Model->ShouldShowNodesAndPinsUnlocalized() ? ECheckBoxState::Unchecked : ECheckBoxState::Checked);
+	UpdateInternalStateFromSettingsModel();
 }
 
 void FInternationalizationSettingsModelDetails::CustomizeDetails( IDetailLayoutBuilder& DetailBuilder )
@@ -160,35 +242,47 @@ void FInternationalizationSettingsModelDetails::CustomizeDetails( IDetailLayoutB
 
 	Model->OnSettingsChanged().AddRaw(this, &FInternationalizationSettingsModelDetails::OnSettingsChanged);
 
-	// If the saved culture is not the same as the actual current culture, a restart is needed to sync them fully and properly.
-	FString SavedEditorCultureName = Model->GetEditorCultureName();
-	if ( !SavedEditorCultureName.IsEmpty() && SavedEditorCultureName != I18N.GetCurrentCulture()->GetName() )
+	// Populate master list of editor cultures.
 	{
-		SelectedEditorCulture = I18N.GetCulture(SavedEditorCultureName);
-		RequiresRestart = true;
+		TArray<FCultureRef> LocalizedCultures;
+		// Use only the cultures for which the editor has localizations for.
+		FInternationalization::Get().GetCulturesWithAvailableLocalization(FPaths::GetEditorLocalizationPaths(), LocalizedCultures, true);
+		for (const FCultureRef& Culture : LocalizedCultures)
+		{
+			AvailableEditorCultures.Add(Culture);
+			const FCulturePtr LanguageCulture = FInternationalization::Get().GetCulture(Culture->GetTwoLetterISOLanguageName());
+			if (LanguageCulture.IsValid())
+			{
+				AvailableEditorLanguages.AddUnique(LanguageCulture);
+			}
+		}
+		AvailableEditorLanguages.Sort(FCompareCultureByNativeLanguage());
 	}
-	else
+	// Populate master list of game cultures.
 	{
-		SelectedEditorCulture = I18N.GetCurrentCulture();
-		RequiresRestart = false;
+		TArray<FCultureRef> LocalizedCultures;
+		// Use only the cultures for which the game has localizations for.
+		FInternationalization::Get().GetCulturesWithAvailableLocalization(FPaths::GetGameLocalizationPaths(), LocalizedCultures, true);
+		for (const FCultureRef& Culture : LocalizedCultures)
+		{
+			AvailableNativeGameCultures.Add(Culture);
+			const FCulturePtr LanguageCulture = FInternationalization::Get().GetCulture(Culture->GetTwoLetterISOLanguageName());
+			AvailableNativeGameLanguages.AddUnique(LanguageCulture);
+		}
+		// Add an option for no game culture.
+		AvailableNativeGameCultures.Add(nullptr);
+		AvailableNativeGameLanguages.AddUnique(nullptr);
+		AvailableNativeGameLanguages.Sort(FCompareCultureByNativeLanguage());
 	}
 
-	const FString SavedNativeGameCultureName = Model->GetNativeGameCultureName();
-	if (!SavedNativeGameCultureName.IsEmpty())
-	{
-		SelectedNativeGameCulture = I18N.GetCulture(SavedNativeGameCultureName);
-		SelectedNativeGameLanguage = SelectedNativeGameCulture.IsValid() ? I18N.GetCulture(SelectedNativeGameCulture->GetTwoLetterISOLanguageName()) : nullptr;
-	}
-	else
+	UpdateInternalStateFromSettingsModel();
+
+	// If the saved editor culture is not the same as the current culture, a restart is needed to sync them fully and properly.
+	if (SelectedEditorCulture != I18N.GetCurrentCulture())
 	{
-		SelectedNativeGameCulture = nullptr;
-		SelectedNativeGameLanguage = nullptr;
+		RequiresRestart = true;
 	}
 
-	// Populate all our cultures
-	RefreshAvailableEditorCultures();
-	RefreshAvailableNativeGameCultures();
-
 	IDetailCategoryBuilder& CategoryBuilder = DetailBuilder.EditCategory("Internationalization");
 
 	{
@@ -404,114 +498,33 @@ void FInternationalizationSettingsModelDetails::CustomizeDetails( IDetailLayoutB
 	];
 }
 
-void FInternationalizationSettingsModelDetails::RefreshAvailableEditorCultures()
-{
-	{
-		AvailableEditorCultures.Empty();
-		TArray<FCultureRef> LocalizedCultures;
-		FInternationalization::Get().GetCulturesWithAvailableLocalization(FPaths::GetEditorLocalizationPaths(), LocalizedCultures, true);
-		for (const FCultureRef& Culture : LocalizedCultures)
-		{
-			AvailableEditorCultures.Add(Culture);
-		}
-
-		// Update our selected culture based on the available choices
-		if ( !AvailableEditorCultures.Contains( SelectedEditorCulture ) )
-		{
-			SelectedEditorCulture = NULL;
-		}
-		RefreshAvailableEditorLanguages();
-		RefreshAvailableEditorRegions();
-	}
-}
-
-void FInternationalizationSettingsModelDetails::RefreshAvailableEditorLanguages()
-{
-	AvailableEditorLanguages.Empty();
-	
-	FString SelectedLanguageName;
-	if ( SelectedEditorCulture.IsValid() )
-	{
-		SelectedLanguageName = SelectedEditorCulture->GetNativeLanguage();
-	}
-
-	// Setup the language list
-	for( const auto& Culture : AvailableEditorCultures )
-	{
-		FCulturePtr LanguageCulture = FInternationalization::Get().GetCulture(Culture->GetTwoLetterISOLanguageName());
-		if (LanguageCulture.IsValid() && AvailableEditorLanguages.Find(LanguageCulture) == INDEX_NONE)
-		{
-			AvailableEditorLanguages.Add(LanguageCulture);
-				
-			// Do we have a match for the base language
-			const FString CultureLanguageName = Culture->GetNativeLanguage();
-			if ( SelectedLanguageName == CultureLanguageName)
-			{
-				SelectedEditorLanguage = Culture;
-			}
-		}
-	}
-
-	AvailableEditorLanguages.Sort( FCompareCultureByNativeLanguage() );
-}
 
 void FInternationalizationSettingsModelDetails::RefreshAvailableEditorRegions()
 {
 	AvailableEditorRegions.Empty();
 
-	FCulturePtr DefaultCulture;
-	if ( SelectedEditorLanguage.IsValid() )
+	// The available regions are based on the selected language, otherwise no regions are available.
+	if (SelectedEditorLanguage.IsValid())
 	{
 		const FString SelectedLanguageName = SelectedEditorLanguage->GetTwoLetterISOLanguageName();
 
-		// Setup the region list
-		for( const auto& Culture : AvailableEditorCultures )
+		// Add any cultures whose language is the the selected language.
+		for (const FCulturePtr& Culture : AvailableEditorCultures)
 		{
 			const FString CultureLanguageName = Culture->GetTwoLetterISOLanguageName();
 			if ( SelectedLanguageName == CultureLanguageName)
 			{
-				AvailableEditorRegions.Add( Culture );
-
-				// If this doesn't have a valid region... assume it's the default
-				const FString CultureRegionName = Culture->GetNativeRegion();
-				if ( CultureRegionName.IsEmpty() )
-				{	
-					DefaultCulture = Culture;
-				}
+				AvailableEditorRegions.AddUnique(Culture);
 			}
 		}
 
-		AvailableEditorRegions.Sort( FCompareCultureByNativeRegion() );
-	}
-
-	// If we have a preferred default (or there's only one in the list), select that now
-	if ( !DefaultCulture.IsValid() )
-	{
-		if(AvailableEditorRegions.Num() == 1)
-		{
-			DefaultCulture = AvailableEditorRegions.Last();
-		}
-	}
-
-	if ( DefaultCulture.IsValid() )
-	{
-		// Set it as our default region, if one hasn't already been chosen
-		if ( !SelectedEditorCulture.IsValid() && EditorRegionComboBox.IsValid() )
-		{
-			// We have to update the combo box like this, otherwise it'll do a null selection when we next click on it
-			EditorRegionComboBox->SetSelectedItem( DefaultCulture );
-		}
-	}
-
-	if ( EditorRegionComboBox.IsValid() )
-	{
-		EditorRegionComboBox->RefreshOptions();
+		AvailableEditorRegions.Sort(FCompareCultureByNativeRegion());
 	}
 }
 
 FText FInternationalizationSettingsModelDetails::GetEditorCurrentLanguageText() const
 {
-	if( SelectedEditorLanguage.IsValid() )
+	if (SelectedEditorLanguage.IsValid())
 	{
 		return FCompareCultureByNativeLanguage::GetCultureNativeLanguageText(SelectedEditorLanguage);
 	}
@@ -520,16 +533,26 @@ FText FInternationalizationSettingsModelDetails::GetEditorCurrentLanguageText()
 
 void FInternationalizationSettingsModelDetails::OnEditorLanguageSelectionChanged( FCulturePtr Culture, ESelectInfo::Type SelectionType )
 {
+	TGuardValue<bool> Guard(IsMakingChangesToModel, true);
 	SelectedEditorLanguage = Culture;
-	SelectedEditorCulture = Culture;
+	RefreshAvailableEditorRegions();
+
+	// Fallback on the first available region of the selected language.
+	SelectedEditorCulture = AvailableEditorRegions[0];
+
+	if (EditorRegionComboBox.IsValid())
+	{
+		EditorRegionComboBox->RefreshOptions();
+		EditorRegionComboBox->SetSelectedItem(SelectedEditorCulture);
+	}
+
 	Model->SetEditorCultureName(SelectedEditorCulture->GetName());
 	RequiresRestart = true;
-	RefreshAvailableEditorRegions();
 }
 
 FText FInternationalizationSettingsModelDetails::GetCurrentEditorRegionText() const
 {
-	if( SelectedEditorCulture.IsValid() )
+	if (SelectedEditorCulture.IsValid())
 	{
 		return FCompareCultureByNativeRegion::GetCultureNativeRegionText(SelectedEditorCulture);
 	}
@@ -539,8 +562,8 @@ FText FInternationalizationSettingsModelDetails::GetCurrentEditorRegionText() co
 void FInternationalizationSettingsModelDetails::OnEditorRegionSelectionChanged( FCulturePtr Culture, ESelectInfo::Type SelectionType )
 {
 	TGuardValue<bool> Guard(IsMakingChangesToModel, true);
-	IsMakingChangesToModel = true;
 	SelectedEditorCulture = Culture;
+
 	Model->SetEditorCultureName(SelectedEditorCulture->GetName());
 	RequiresRestart = true;
 }
@@ -550,109 +573,23 @@ bool FInternationalizationSettingsModelDetails::IsEditorRegionSelectionAllowed()
 	return SelectedEditorLanguage.IsValid();
 }
 
-void FInternationalizationSettingsModelDetails::RefreshAvailableNativeGameCultures()
-{
-	{
-		AvailableNativeGameCultures.Empty();
-		TArray<FCultureRef> LocalizedCultures;
-		FInternationalization::Get().GetCulturesWithAvailableLocalization(FPaths::GetGameLocalizationPaths(), LocalizedCultures, true);
-		for (const FCultureRef& Culture : LocalizedCultures)
-		{
-			AvailableNativeGameCultures.Add(Culture);
-		}
-		AvailableNativeGameCultures.Add(nullptr);
-
-		// Update our selected culture based on the available choices
-		if ( !AvailableNativeGameCultures.Contains( SelectedNativeGameCulture ) )
-		{
-			SelectedNativeGameCulture = NULL;
-		}
-		RefreshAvailableNativeGameLanguages();
-		RefreshAvailableNativeGameRegions();
-	}
-}
-
-void FInternationalizationSettingsModelDetails::RefreshAvailableNativeGameLanguages()
-{
-	AvailableNativeGameLanguages.Empty();
-
-	FString SelectedLanguageName;
-	if ( SelectedNativeGameCulture.IsValid() )
-	{
-		SelectedLanguageName = SelectedNativeGameCulture->GetNativeLanguage();
-	}
-
-	// Setup the language list
-	for( const auto& Culture : AvailableNativeGameCultures )
-	{
-		FCulturePtr LanguageCulture = Culture.IsValid() ? FInternationalization::Get().GetCulture(Culture->GetTwoLetterISOLanguageName()) : nullptr;
-		if (AvailableNativeGameLanguages.Find(LanguageCulture) == INDEX_NONE)
-		{
-			AvailableNativeGameLanguages.Add(LanguageCulture);
-
-			// Do we have a match for the base language
-			const FString CultureLanguageName = Culture.IsValid() ? Culture->GetNativeLanguage() : TEXT("");
-			if ( SelectedLanguageName == CultureLanguageName)
-			{
-				SelectedNativeGameLanguage = Culture;
-			}
-		}
-	}
-
-	AvailableNativeGameLanguages.Sort( FCompareCultureByNativeLanguage() );
-}
-
 void FInternationalizationSettingsModelDetails::RefreshAvailableNativeGameRegions()
 {
 	AvailableNativeGameRegions.Empty();
 
-	FCulturePtr DefaultCulture;
-	{
-		const FString SelectedLanguageName = SelectedNativeGameLanguage.IsValid() ? SelectedNativeGameLanguage->GetTwoLetterISOLanguageName() : TEXT("");
+	const FString SelectedLanguageName = SelectedNativeGameLanguage.IsValid() ? SelectedNativeGameLanguage->GetTwoLetterISOLanguageName() : TEXT("");
 
-		// Setup the region list
-		for( const auto& Culture : AvailableNativeGameCultures )
-		{
-			const FString CultureLanguageName = Culture.IsValid() ? Culture->GetTwoLetterISOLanguageName() : TEXT("");
-			if ( SelectedLanguageName == CultureLanguageName)
-			{
-				AvailableNativeGameRegions.Add( Culture );
-
-				// If this doesn't have a valid region... assume it's the default
-				const FString CultureRegionName = Culture.IsValid() ? Culture->GetNativeRegion() : TEXT("");
-				if ( CultureRegionName.IsEmpty() )
-				{	
-					DefaultCulture = Culture;
-				}
-			}
-		}
-
-		AvailableNativeGameRegions.Sort( FCompareCultureByNativeRegion() );
-	}
-
-	// If we have a preferred default (or there's only one in the list), select that now
-	if ( !DefaultCulture.IsValid() )
+	// Setup the region list
+	for (const FCulturePtr& Culture : AvailableNativeGameCultures)
 	{
-		if(AvailableNativeGameRegions.Num() == 1)
+		const FString CultureLanguageName = Culture.IsValid() ? Culture->GetTwoLetterISOLanguageName() : TEXT("");
+		if ( SelectedLanguageName == CultureLanguageName)
 		{
-			DefaultCulture = AvailableNativeGameRegions.Last();
+			AvailableNativeGameRegions.AddUnique(Culture);
 		}
 	}
 
-	if ( DefaultCulture.IsValid() )
-	{
-		// Set it as our default region, if one hasn't already been chosen
-		if ( !SelectedNativeGameCulture.IsValid() && NativeGameRegionComboBox.IsValid() )
-		{
-			// We have to update the combo box like this, otherwise it'll do a null selection when we next click on it
-			NativeGameRegionComboBox->SetSelectedItem( DefaultCulture );
-		}
-	}
-
-	if ( NativeGameRegionComboBox.IsValid() )
-	{
-		NativeGameRegionComboBox->RefreshOptions();
-	}
+	AvailableNativeGameRegions.Sort(FCompareCultureByNativeRegion());
 }
 
 FText FInternationalizationSettingsModelDetails::GetNativeGameCurrentLanguageText() const
@@ -668,9 +605,19 @@ void FInternationalizationSettingsModelDetails::OnNativeGameLanguageSelectionCha
 {
 	TGuardValue<bool> Guard(IsMakingChangesToModel, true);
 	SelectedNativeGameLanguage = Culture;
-	SelectedNativeGameCulture = Culture;
-	Model->SetNativeGameCultureName(SelectedNativeGameCulture.IsValid() ? SelectedNativeGameCulture->GetName() : TEXT(""));
 	RefreshAvailableNativeGameRegions();
+
+	// Fallback on the first available region of the selected language.
+	SelectedNativeGameCulture = AvailableEditorRegions[0];
+
+	if (NativeGameRegionComboBox.IsValid())
+	{
+		NativeGameRegionComboBox->RefreshOptions();
+		NativeGameRegionComboBox->SetSelectedItem(SelectedNativeGameCulture);
+	}
+
+	Model->SetNativeGameCultureName(SelectedNativeGameCulture.IsValid() ? SelectedNativeGameCulture->GetName() : TEXT(""));
+	RequiresRestart = true;
 }
 
 FText FInternationalizationSettingsModelDetails::GetCurrentNativeGameRegionText() const
@@ -686,6 +633,7 @@ void FInternationalizationSettingsModelDetails::OnNativeGameRegionSelectionChang
 {
 	TGuardValue<bool> Guard(IsMakingChangesToModel, true);
 	SelectedNativeGameCulture = Culture;
+
 	Model->SetNativeGameCultureName(SelectedNativeGameCulture.IsValid() ? SelectedNativeGameCulture->GetName() : TEXT(""));
 	RequiresRestart = true;
 }
diff --git a/Engine/Source/Developer/InternationalizationSettings/Public/InternationalizationSettingsModelDetails.h b/Engine/Source/Developer/InternationalizationSettings/Public/InternationalizationSettingsModelDetails.h
--- a/Engine/Source/Developer/InternationalizationSettings/Public/InternationalizationSettingsModelDetails.h
+++ b/Engine/Source/Developer/InternationalizationSettings/Public/InternationalizationSettingsModelDetails.h
@@ -24,19 +24,11 @@ public:
 	virtual void CustomizeDetails( IDetailLayoutBuilder& DetailLayout ) override;
 
 private:
+	void UpdateInternalStateFromSettingsModel();
 
 	void OnSettingsChanged();
 
-	/** Prompt the user to restart the editor */
-	void PromptForRestart() const;
-
-	/** Look for available cultures */
-	void RefreshAvailableEditorCultures();
-
-	/** Look for available languages */
-	void RefreshAvailableEditorLanguages();
-
-	/** Look for available regions */
+	/** Look for available editor culture regions */
 	void RefreshAvailableEditorRegions();
 
 	/** Delegate called to display the current editor language */
@@ -54,25 +46,19 @@ private:
 	/** Delegate called to see if we're allowed to change the region */
 	bool IsEditorRegionSelectionAllowed() const;
 
-	/** Look for available cultures */
-	void RefreshAvailableNativeGameCultures();
-
-	/** Look for available languages */
-	void RefreshAvailableNativeGameLanguages();
-
-	/** Look for available regions */
+	/** Look for available game culture regions */
 	void RefreshAvailableNativeGameRegions();
 
-	/** Delegate called to display the current NativeGame language */
+	/** Delegate called to display the current native game language */
 	FText GetNativeGameCurrentLanguageText() const;
 
-	/** Delegate called when the NativeGame language selection has changed */
+	/** Delegate called when the native game language selection has changed */
 	void OnNativeGameLanguageSelectionChanged( FCulturePtr Culture, ESelectInfo::Type SelectionType );
 
-	/** Delegate called to display the current NativeGame region */
+	/** Delegate called to display the current native game region */
 	FText GetCurrentNativeGameRegionText() const;
 
-	/** Delegate called when the NativeGame region selection has changed */
+	/** Delegate called when the native game region selection has changed */
 	void OnNativeGameRegionSelectionChanged( FCulturePtr Culture, ESelectInfo::Type SelectionType );
 
 	/** Delegate called to see if we're allowed to change the region */
@@ -90,9 +76,6 @@ private:
 	/** Delegate called when the the checked state of whether or not field names should be localized has changed. */
 	void ShouldLoadLocalizedFieldNamesCheckChanged(ECheckBoxState CheckState);
 
-	/** Write to config now the Editor is shutting down (all packages are saved) */
-	void HandleShutdownPostPackagesSaved();
-
 	/** Delegate called when the the checked state of whether or not nodes and pins in graph editors should be localized has changed. */
 	void ShouldShowNodesAndPinsUnlocalized(ECheckBoxState CheckState);
 
