Added toast prompt when making changes to language settings which require restart.

diff --git a/Engine/Source/Developer/InternationalizationSettings/InternationalizationSettings.Build.cs b/Engine/Source/Developer/InternationalizationSettings/InternationalizationSettings.Build.cs
--- a/Engine/Source/Developer/InternationalizationSettings/InternationalizationSettings.Build.cs
+++ b/Engine/Source/Developer/InternationalizationSettings/InternationalizationSettings.Build.cs
@@ -23,8 +23,15 @@ namespace UnrealBuildTool.Rules
 			PrivateIncludePathModuleNames.AddRange(
 				new string[] {
 					"Settings",
+                    "SettingsEditor"
 				}
 			);
+
+            DynamicallyLoadedModuleNames.AddRange(
+                new string[] {
+				    "SettingsEditor"
+                }
+            );
 		}
 	}
 }
diff --git a/Engine/Source/Developer/InternationalizationSettings/Private/InternationalizationSettingsModelDetails.cpp b/Engine/Source/Developer/InternationalizationSettings/Private/InternationalizationSettingsModelDetails.cpp
--- a/Engine/Source/Developer/InternationalizationSettings/Private/InternationalizationSettingsModelDetails.cpp
+++ b/Engine/Source/Developer/InternationalizationSettings/Private/InternationalizationSettingsModelDetails.cpp
@@ -2,6 +2,7 @@
 
 #include "InternationalizationSettingsModulePrivatePCH.h"
 #include "EdGraph/EdGraphSchema.h"
+#include "ISettingsEditorModule.h"
 
 #define LOCTEXT_NAMESPACE "InternationalizationSettingsModelDetails"
 
@@ -548,6 +549,9 @@ void FInternationalizationSettingsModelDetails::OnEditorLanguageSelectionChanged
 
 	Model->SetEditorCultureName(SelectedEditorCulture->GetName());
 	RequiresRestart = true;
+
+	ISettingsEditorModule& SettingsEditorModule = FModuleManager::GetModuleChecked<ISettingsEditorModule>("SettingsEditor");
+	SettingsEditorModule.OnApplicationRestartRequired();
 }
 
 FText FInternationalizationSettingsModelDetails::GetCurrentEditorRegionText() const
@@ -566,6 +570,9 @@ void FInternationalizationSettingsModelDetails::OnEditorRegionSelectionChanged(
 
 	Model->SetEditorCultureName(SelectedEditorCulture->GetName());
 	RequiresRestart = true;
+
+	ISettingsEditorModule& SettingsEditorModule = FModuleManager::GetModuleChecked<ISettingsEditorModule>("SettingsEditor");
+	SettingsEditorModule.OnApplicationRestartRequired();
 }
 
 bool FInternationalizationSettingsModelDetails::IsEditorRegionSelectionAllowed() const
@@ -618,6 +625,9 @@ void FInternationalizationSettingsModelDetails::OnNativeGameLanguageSelectionCha
 
 	Model->SetNativeGameCultureName(SelectedNativeGameCulture.IsValid() ? SelectedNativeGameCulture->GetName() : TEXT(""));
 	RequiresRestart = true;
+
+	ISettingsEditorModule& SettingsEditorModule = FModuleManager::GetModuleChecked<ISettingsEditorModule>("SettingsEditor");
+	SettingsEditorModule.OnApplicationRestartRequired();
 }
 
 FText FInternationalizationSettingsModelDetails::GetCurrentNativeGameRegionText() const
@@ -636,6 +646,9 @@ void FInternationalizationSettingsModelDetails::OnNativeGameRegionSelectionChang
 
 	Model->SetNativeGameCultureName(SelectedNativeGameCulture.IsValid() ? SelectedNativeGameCulture->GetName() : TEXT(""));
 	RequiresRestart = true;
+
+	ISettingsEditorModule& SettingsEditorModule = FModuleManager::GetModuleChecked<ISettingsEditorModule>("SettingsEditor");
+	SettingsEditorModule.OnApplicationRestartRequired();
 }
 
 bool FInternationalizationSettingsModelDetails::IsNativeGameRegionSelectionAllowed() const
