Fixed crash on Windows Chinese version when changing editor language

Made the Region & Language UI more hardy against null regions when changing your editor language.

diff --git a/Engine/Source/Developer/InternationalizationSettings/Private/InternationalizationSettingsModelDetails.cpp b/Engine/Source/Developer/InternationalizationSettings/Private/InternationalizationSettingsModelDetails.cpp
--- a/Engine/Source/Developer/InternationalizationSettings/Private/InternationalizationSettingsModelDetails.cpp
+++ b/Engine/Source/Developer/InternationalizationSettings/Private/InternationalizationSettingsModelDetails.cpp
@@ -107,6 +107,7 @@ void FInternationalizationSettingsModelDetails::UpdateInternalStateFromSettingsM
 	FInternationalization& I18N = FInternationalization::Get();
 
 	const FString SavedEditorCultureName = Model->GetEditorCultureName();
+
 	// Attempt to use the culture specified in the settings.
 	if (!SavedEditorCultureName.IsEmpty())
 	{
@@ -127,7 +128,7 @@ void FInternationalizationSettingsModelDetails::UpdateInternalStateFromSettingsM
 			// Fallback on the first available region of the selected language.
 			else
 			{
-				SelectedEditorCulture = AvailableEditorRegions[0];
+				SelectedEditorLanguage = (AvailableEditorRegions.Num() > 0) ? AvailableEditorRegions[0] : nullptr;
 			}
 		}
 		else
@@ -135,6 +136,12 @@ void FInternationalizationSettingsModelDetails::UpdateInternalStateFromSettingsM
 			SelectedEditorCulture = I18N.GetCurrentCulture();
 			SelectedEditorLanguage = SelectedEditorCulture.IsValid() ? I18N.GetCulture(SelectedEditorCulture->GetTwoLetterISOLanguageName()) : nullptr;
 			RefreshAvailableEditorRegions();
+
+			if (!SelectedEditorLanguage.IsValid())
+			{
+				// Fallback on the first available region of the selected language.
+				SelectedEditorLanguage = (AvailableEditorRegions.Num() > 0) ? AvailableEditorRegions[0] : nullptr;
+			}
 		}
 	}
 	// Fallback on the current culture of the editor.
@@ -143,13 +150,20 @@ void FInternationalizationSettingsModelDetails::UpdateInternalStateFromSettingsM
 		SelectedEditorCulture = I18N.GetCurrentCulture();
 		SelectedEditorLanguage = SelectedEditorCulture.IsValid() ? I18N.GetCulture(SelectedEditorCulture->GetTwoLetterISOLanguageName()) : nullptr;
 		RefreshAvailableEditorRegions();
+
+		if (!SelectedEditorLanguage.IsValid())
+		{
+			// Fallback on the first available region of the selected language.
+			SelectedEditorLanguage = (AvailableEditorRegions.Num() > 0) ? AvailableEditorRegions[0] : nullptr;
+		}
 	}
-	check(AvailableEditorLanguages.Contains(SelectedEditorLanguage) && AvailableEditorRegions.Contains(SelectedEditorCulture));
+	
 	if (EditorLanguageComboBox.IsValid())
 	{
 		EditorLanguageComboBox->RefreshOptions();
 		EditorLanguageComboBox->SetSelectedItem(SelectedEditorLanguage);
 	}
+
 	if (EditorRegionComboBox.IsValid())
 	{
 		EditorRegionComboBox->RefreshOptions();
@@ -157,6 +171,7 @@ void FInternationalizationSettingsModelDetails::UpdateInternalStateFromSettingsM
 	}
 
 	const FString SavedNativeGameCultureName = Model->GetNativeGameCultureName();
+
 	// Attempt to use the culture specified in the settings.
 	if (!SavedNativeGameCultureName.IsEmpty())
 	{
@@ -177,7 +192,7 @@ void FInternationalizationSettingsModelDetails::UpdateInternalStateFromSettingsM
 			// Fallback on the first available region of the selected language.
 			else
 			{
-				SelectedNativeGameCulture = AvailableNativeGameRegions[0];
+				SelectedNativeGameCulture = (AvailableNativeGameRegions.Num() > 0) ? AvailableNativeGameRegions[0] : nullptr;
 			}
 		}
 		else
@@ -185,6 +200,12 @@ void FInternationalizationSettingsModelDetails::UpdateInternalStateFromSettingsM
 			SelectedNativeGameCulture = I18N.GetCurrentCulture();
 			SelectedNativeGameLanguage = SelectedNativeGameCulture.IsValid() ? I18N.GetCulture(SelectedNativeGameCulture->GetTwoLetterISOLanguageName()) : nullptr;
 			RefreshAvailableNativeGameRegions();
+
+			if (!SelectedNativeGameCulture.IsValid())
+			{
+				// Fallback on the first available region of the selected language.
+				SelectedNativeGameCulture = (AvailableNativeGameRegions.Num() > 0) ? AvailableNativeGameRegions[0] : nullptr;
+			}
 		}
 	}
 	// Fallback to no specified language or region.
@@ -194,12 +215,13 @@ void FInternationalizationSettingsModelDetails::UpdateInternalStateFromSettingsM
 		SelectedNativeGameLanguage = nullptr;
 		RefreshAvailableNativeGameRegions();
 	}
-	check(AvailableNativeGameLanguages.Contains(SelectedNativeGameLanguage) && AvailableNativeGameRegions.Contains(SelectedNativeGameCulture));
+
 	if (NativeGameLanguageComboBox.IsValid())
 	{
 		NativeGameLanguageComboBox->RefreshOptions();
 		NativeGameLanguageComboBox->SetSelectedItem(SelectedNativeGameLanguage);
 	}
+
 	if (NativeGameRegionComboBox.IsValid())
 	{
 		NativeGameRegionComboBox->RefreshOptions();
@@ -210,6 +232,7 @@ void FInternationalizationSettingsModelDetails::UpdateInternalStateFromSettingsM
 	{
 		LocalizedPropertyNamesCheckBox->SetIsChecked(Model->ShouldLoadLocalizedPropertyNames() ? ECheckBoxState::Checked : ECheckBoxState::Unchecked);
 	}
+
 	if (UnlocalizedNodesAndPinsCheckBox.IsValid())
 	{
 		UnlocalizedNodesAndPinsCheckBox->SetIsChecked(Model->ShouldShowNodesAndPinsUnlocalized() ? ECheckBoxState::Unchecked : ECheckBoxState::Checked);
@@ -279,7 +302,7 @@ void FInternationalizationSettingsModelDetails::CustomizeDetails( IDetailLayoutB
 	UpdateInternalStateFromSettingsModel();
 
 	// If the saved editor culture is not the same as the current culture, a restart is needed to sync them fully and properly.
-	if (SelectedEditorCulture != I18N.GetCurrentCulture())
+	if (SelectedEditorCulture.IsValid() && SelectedEditorCulture != I18N.GetCurrentCulture())
 	{
 		RequiresRestart = true;
 	}
@@ -529,7 +552,7 @@ FText FInternationalizationSettingsModelDetails::GetEditorCurrentLanguageText()
 	{
 		return FCompareCultureByNativeLanguage::GetCultureNativeLanguageText(SelectedEditorLanguage);
 	}
-	return FText::GetEmpty();
+	return LOCTEXT("None", "(None)");
 }
 
 void FInternationalizationSettingsModelDetails::OnEditorLanguageSelectionChanged( FCulturePtr Culture, ESelectInfo::Type SelectionType )
@@ -539,7 +562,7 @@ void FInternationalizationSettingsModelDetails::OnEditorLanguageSelectionChanged
 	RefreshAvailableEditorRegions();
 
 	// Fallback on the first available region of the selected language.
-	SelectedEditorCulture = AvailableEditorRegions[0];
+	SelectedEditorCulture = (AvailableEditorRegions.Num() > 0) ? AvailableEditorRegions[0] : nullptr;
 
 	if (EditorRegionComboBox.IsValid())
 	{
@@ -547,7 +570,7 @@ void FInternationalizationSettingsModelDetails::OnEditorLanguageSelectionChanged
 		EditorRegionComboBox->SetSelectedItem(SelectedEditorCulture);
 	}
 
-	Model->SetEditorCultureName(SelectedEditorCulture->GetName());
+	Model->SetEditorCultureName(SelectedEditorCulture.IsValid() ? SelectedEditorCulture->GetName() : TEXT(""));
 	RequiresRestart = true;
 
 	ISettingsEditorModule& SettingsEditorModule = FModuleManager::GetModuleChecked<ISettingsEditorModule>("SettingsEditor");
@@ -560,7 +583,7 @@ FText FInternationalizationSettingsModelDetails::GetCurrentEditorRegionText() co
 	{
 		return FCompareCultureByNativeRegion::GetCultureNativeRegionText(SelectedEditorCulture);
 	}
-	return FText::GetEmpty();
+	return LOCTEXT("None", "(None)");
 }
 
 void FInternationalizationSettingsModelDetails::OnEditorRegionSelectionChanged( FCulturePtr Culture, ESelectInfo::Type SelectionType )
@@ -568,7 +591,7 @@ void FInternationalizationSettingsModelDetails::OnEditorRegionSelectionChanged(
 	TGuardValue<bool> Guard(IsMakingChangesToModel, true);
 	SelectedEditorCulture = Culture;
 
-	Model->SetEditorCultureName(SelectedEditorCulture->GetName());
+	Model->SetEditorCultureName(SelectedEditorCulture.IsValid() ? SelectedEditorCulture->GetName() : TEXT(""));
 	RequiresRestart = true;
 
 	ISettingsEditorModule& SettingsEditorModule = FModuleManager::GetModuleChecked<ISettingsEditorModule>("SettingsEditor");
