Largely rewritten Editor Region & Language Settings page. Far less error prone compared to the old implementation. I'd like to do a lot more to make this not so lame in the future - move the game's native culture to something not editor-specific and leverage it as the default for new game targets in the localization dashboard.

diff --git a/Engine/Source/Developer/InternationalizationSettings/Classes/InternationalizationSettingsModel.h b/Engine/Source/Developer/InternationalizationSettings/Classes/InternationalizationSettingsModel.h
--- a/Engine/Source/Developer/InternationalizationSettings/Classes/InternationalizationSettingsModel.h
+++ b/Engine/Source/Developer/InternationalizationSettings/Classes/InternationalizationSettingsModel.h
@@ -8,36 +8,20 @@
 /**
  * Implements loading and saving of internationalization settings.
  */
-UCLASS()
+UCLASS(config=EditorSettings)
 class INTERNATIONALIZATIONSETTINGS_API UInternationalizationSettingsModel
-	:	public UObject
+	: public UObject
 {
 	GENERATED_UCLASS_BODY()
 
 public:
-
-	void SaveDefaults();
 	void ResetToDefault();
-	FString GetEditorCultureName() const;
+	bool GetEditorCultureName(FString& OutEditorCultureName) const;
 	void SetEditorCultureName(const FString& CultureName);
-	FString GetNativeGameCultureName() const;
+	bool GetNativeGameCultureName(FString& OutNativeGameCultureName) const;
 	void SetNativeGameCultureName(const FString& CultureName);
 	bool ShouldLoadLocalizedPropertyNames() const;
 	void ShouldLoadLocalizedPropertyNames(const bool Value);
 	bool ShouldShowNodesAndPinsUnlocalized() const;
 	void ShouldShowNodesAndPinsUnlocalized(const bool Value);
-public:
-
-	/**
-	 * Returns an event delegate that is executed when a setting has changed.
-	 *
-	 * @return The delegate.
-	 */
-	DECLARE_EVENT(UInternationalizationSettingsModel, FSettingsChangedEvent);
-	FSettingsChangedEvent& OnSettingsChanged( ) { return SettingsChangedEvent; }
-
-private:
-
-	// Holds an event delegate that is executed when a setting has changed.
-	FSettingsChangedEvent SettingsChangedEvent;
-};
+};
\ No newline at end of file
diff --git a/Engine/Source/Developer/InternationalizationSettings/InternationalizationSettings.Build.cs b/Engine/Source/Developer/InternationalizationSettings/InternationalizationSettings.Build.cs
--- a/Engine/Source/Developer/InternationalizationSettings/InternationalizationSettings.Build.cs
+++ b/Engine/Source/Developer/InternationalizationSettings/InternationalizationSettings.Build.cs
@@ -17,6 +17,7 @@ namespace UnrealBuildTool.Rules
 				    "EditorStyle",
 				    "PropertyEditor",
 				    "SharedSettingsWidgets",
+                    "Localization",
                 }
             );
 
diff --git a/Engine/Source/Developer/InternationalizationSettings/Private/InternationalizationSettingsModel.cpp b/Engine/Source/Developer/InternationalizationSettings/Private/InternationalizationSettingsModel.cpp
--- a/Engine/Source/Developer/InternationalizationSettings/Private/InternationalizationSettingsModel.cpp
+++ b/Engine/Source/Developer/InternationalizationSettings/Private/InternationalizationSettingsModel.cpp
@@ -9,97 +9,67 @@ UInternationalizationSettingsModel::UInternationalizationSettingsModel( const FO
 	: Super(ObjectInitializer)
 { }
 
-void UInternationalizationSettingsModel::SaveDefaults()
-{
-	FString SavedCultureName;
-	GConfig->GetString( TEXT("Internationalization"), TEXT("Culture"), SavedCultureName, GEditorSettingsIni );
-	GConfig->SetString( TEXT("Internationalization"), TEXT("Culture"), *SavedCultureName, GEngineIni );
-
-	bool bShouldLoadLocalizedPropertyNames = true;
-	GConfig->GetBool( TEXT("Internationalization"), TEXT("ShouldLoadLocalizedPropertyNames"), bShouldLoadLocalizedPropertyNames, GEditorSettingsIni );
-	GConfig->SetBool( TEXT("Internationalization"), TEXT("ShouldLoadLocalizedPropertyNames"), bShouldLoadLocalizedPropertyNames, GEngineIni );
-
-	bool bShowNodesAndPinsUnlocalized = false;
-	GConfig->GetBool( TEXT("Internationalization"), TEXT("ShowNodesAndPinsUnlocalized"), bShowNodesAndPinsUnlocalized, GEditorSettingsIni );
-	GConfig->SetBool( TEXT("Internationalization"), TEXT("ShowNodesAndPinsUnlocalized"), bShowNodesAndPinsUnlocalized, GEngineIni );
-}
-
 void UInternationalizationSettingsModel::ResetToDefault()
 {
+	// Inherit editor culture from engine settings. Empty otherwise.
 	FString SavedCultureName;
-	GConfig->GetString( TEXT("Internationalization"), TEXT("Culture"), SavedCultureName, GEngineIni );
+	GConfig->GetString(TEXT("Internationalization"), TEXT("Culture"), SavedCultureName, GEngineIni);
 	GConfig->SetString( TEXT("Internationalization"), TEXT("Culture"), *SavedCultureName, GEditorSettingsIni );
 
 	GConfig->SetString( TEXT("Internationalization"), TEXT("NativeGameCulture"), TEXT(""), GEditorSettingsIni );
 
-	bool bShouldLoadLocalizedPropertyNames = true;
-	GConfig->GetBool( TEXT("Internationalization"), TEXT("ShouldLoadLocalizedPropertyNames"), bShouldLoadLocalizedPropertyNames, GEngineIni );
-	GConfig->SetBool( TEXT("Internationalization"), TEXT("ShouldLoadLocalizedPropertyNames"), bShouldLoadLocalizedPropertyNames, GEditorSettingsIni );
+	GConfig->SetBool( TEXT("Internationalization"), TEXT("ShouldLoadLocalizedPropertyNames"), true, GEditorSettingsIni );
 
-	bool bShowNodesAndPinsUnlocalized = false;
-	GConfig->GetBool( TEXT("Internationalization"), TEXT("ShowNodesAndPinsUnlocalized"), bShowNodesAndPinsUnlocalized, GEngineIni );
-	GConfig->SetBool( TEXT("Internationalization"), TEXT("ShowNodesAndPinsUnlocalized"), bShowNodesAndPinsUnlocalized, GEditorSettingsIni );
+	GConfig->SetBool( TEXT("Internationalization"), TEXT("ShowNodesAndPinsUnlocalized"), false, GEditorSettingsIni );
 
-	SettingsChangedEvent.Broadcast();
+	GConfig->Flush(false, GEditorSettingsIni);
 }
 
-FString UInternationalizationSettingsModel::GetEditorCultureName() const
+bool UInternationalizationSettingsModel::GetEditorCultureName(FString& OutEditorCultureName) const
 {
-	FString SavedCultureName;
-	if( !GConfig->GetString( TEXT("Internationalization"), TEXT("Culture"), SavedCultureName, GEditorSettingsIni ) )
-	{
-		GConfig->GetString( TEXT("Internationalization"), TEXT("Culture"), SavedCultureName, GEngineIni );
-	}
-	return SavedCultureName;
+	return GConfig->GetString(TEXT("Internationalization"), TEXT("Culture"), OutEditorCultureName, GEditorSettingsIni)
+		|| GConfig->GetString(TEXT("Internationalization"), TEXT("Culture"), OutEditorCultureName, GEngineIni);
 }
 
 void UInternationalizationSettingsModel::SetEditorCultureName(const FString& CultureName)
 {
 	GConfig->SetString( TEXT("Internationalization"), TEXT("Culture"), *CultureName, GEditorSettingsIni );
-	SettingsChangedEvent.Broadcast();
+	GConfig->Flush(false, GEditorSettingsIni);
 }
 
-FString UInternationalizationSettingsModel::GetNativeGameCultureName() const
+bool UInternationalizationSettingsModel::GetNativeGameCultureName(FString& OutNativeGameCultureName) const
 {
-	FString SavedCultureName;
-	GConfig->GetString( TEXT("Internationalization"), TEXT("NativeGameCulture"), SavedCultureName, GEditorSettingsIni );
-	return SavedCultureName;
+	return GConfig->GetString(TEXT("Internationalization"), TEXT("NativeGameCulture"), OutNativeGameCultureName, GEditorSettingsIni);
 }
 
 void UInternationalizationSettingsModel::SetNativeGameCultureName(const FString& CultureName)
 {
 	GConfig->SetString( TEXT("Internationalization"), TEXT("NativeGameCulture"), *CultureName, GEditorSettingsIni );
-	SettingsChangedEvent.Broadcast();
+	GConfig->Flush(false, GEditorSettingsIni);
 }
 
 bool UInternationalizationSettingsModel::ShouldLoadLocalizedPropertyNames() const
 {
 	bool bShouldLoadLocalizedPropertyNames = true;
-	if( !GConfig->GetBool( TEXT("Internationalization"), TEXT("ShouldLoadLocalizedPropertyNames"), bShouldLoadLocalizedPropertyNames, GEditorSettingsIni ) )
-	{
-		GConfig->GetBool( TEXT("Internationalization"), TEXT("ShouldLoadLocalizedPropertyNames"), bShouldLoadLocalizedPropertyNames, GEngineIni );
-	}
+	GConfig->GetBool(TEXT("Internationalization"), TEXT("ShouldLoadLocalizedPropertyNames"), bShouldLoadLocalizedPropertyNames, GEditorSettingsIni);
 	return bShouldLoadLocalizedPropertyNames;
 }
 
 void UInternationalizationSettingsModel::ShouldLoadLocalizedPropertyNames(const bool Value)
 {
 	GConfig->SetBool( TEXT("Internationalization"), TEXT("ShouldLoadLocalizedPropertyNames"), Value, GEditorSettingsIni );
-	SettingsChangedEvent.Broadcast();
+	GConfig->Flush(false, GEditorSettingsIni);
 }
 
 bool UInternationalizationSettingsModel::ShouldShowNodesAndPinsUnlocalized() const
 {
 	bool bShowNodesAndPinsUnlocalized = false;
-	if( !GConfig->GetBool( TEXT("Internationalization"), TEXT("ShowNodesAndPinsUnlocalized"), bShowNodesAndPinsUnlocalized, GEditorSettingsIni ) )
-	{
-		GConfig->GetBool( TEXT("Internationalization"), TEXT("ShowNodesAndPinsUnlocalized"), bShowNodesAndPinsUnlocalized, GEngineIni );
-	}
+	GConfig->GetBool(TEXT("Internationalization"), TEXT("ShowNodesAndPinsUnlocalized"), bShowNodesAndPinsUnlocalized, GEditorSettingsIni);
 	return bShowNodesAndPinsUnlocalized;
 }
 
 void UInternationalizationSettingsModel::ShouldShowNodesAndPinsUnlocalized(const bool Value)
 {
 	GConfig->SetBool( TEXT("Internationalization"), TEXT("ShowNodesAndPinsUnlocalized"), Value, GEditorSettingsIni );
-	SettingsChangedEvent.Broadcast();
+	GConfig->Flush(false, GEditorSettingsIni);
 }
diff --git a/Engine/Source/Developer/InternationalizationSettings/Private/InternationalizationSettingsModelDetails.cpp b/Engine/Source/Developer/InternationalizationSettings/Private/InternationalizationSettingsModelDetails.cpp
--- a/Engine/Source/Developer/InternationalizationSettings/Private/InternationalizationSettingsModelDetails.cpp
+++ b/Engine/Source/Developer/InternationalizationSettings/Private/InternationalizationSettingsModelDetails.cpp
@@ -2,723 +2,336 @@
 
 #include "InternationalizationSettingsModulePrivatePCH.h"
 #include "EdGraph/EdGraphSchema.h"
+#include "SCulturePicker.h"
 #include "ISettingsEditorModule.h"
 
 #define LOCTEXT_NAMESPACE "InternationalizationSettingsModelDetails"
 
-/** Functions for sorting the languages */
-struct FCompareCultureByNativeLanguage
+TSharedRef<IDetailCustomization> FInternationalizationSettingsModelDetails::MakeInstance()
 {
-	static FText GetCultureNativeLanguageText( const FCulturePtr Culture )
-	{
-		return Culture.IsValid() ? FText::FromString(Culture->GetNativeLanguage()) : LOCTEXT("None", "(None)");
-	}
-
-	FORCEINLINE bool operator()( const FCulturePtr A, const FCulturePtr B ) const
-	{
-		const FText AText = GetCultureNativeLanguageText(A);
-		const FText BText = GetCultureNativeLanguageText(B);
-
-		// (None) should appear before all else.
-		if(AText.IdenticalTo(LOCTEXT("None", "(None)")))
-		{
-			return true;
-		}
-		// (None) should appear before all else.
-		if(BText.IdenticalTo(LOCTEXT("None", "(None)")))
-		{
-			return false;
-		}
-		return( GetCultureNativeLanguageText( A ).CompareToCaseIgnored( GetCultureNativeLanguageText( B ) ) ) < 0;
-	}
-};
-
+	return MakeShareable(new FInternationalizationSettingsModelDetails);
+}
 
-/** Functions for sorting the regions */
-struct FCompareCultureByNativeRegion
+namespace
 {
-	static FText GetCultureNativeRegionText( const FCulturePtr Culture )
+	struct FLocalizedCulturesFlyweight
 	{
-		FString Region;
-		if (Culture.IsValid())
-		{
-			Region = Culture->GetNativeRegion();
-			if ( Region.IsEmpty() )
-			{
-				// Fallback to displaying the language, if no region is available
-				return LOCTEXT("NoSpecificRegionOption", "Non-Specific Region");
-			}
-			else
-			{
-				return FText::FromString(Region);
-			}
-		}
-		else
+		TArray<FCultureRef> LocalizedCulturesForEditor;
+		TArray<FCultureRef> LocalizedCulturesForGame;
+
+		FLocalizedCulturesFlyweight()
 		{
-			return LOCTEXT("None", "(None)");
+			FInternationalization::Get().GetCulturesWithAvailableLocalization(FPaths::GetEditorLocalizationPaths(), LocalizedCulturesForEditor, true);
+			FInternationalization::Get().GetCulturesWithAvailableLocalization(FPaths::GetGameLocalizationPaths(), LocalizedCulturesForGame, true);
 		}
-	}
+	};
 
-	FORCEINLINE bool operator()( const FCulturePtr A, const FCulturePtr B ) const
+	class SEditorCultureComboButton : public SCompoundWidget
 	{
-		const FText AText = GetCultureNativeRegionText(A);
-		const FText BText = GetCultureNativeRegionText(B);
+		SLATE_BEGIN_ARGS(SEditorCultureComboButton){}
+		SLATE_END_ARGS()
 
-		// (None) should appear before all else.
-		if(AText.IdenticalTo(LOCTEXT("None", "(None)")))
+	public:
+		void Construct(const FArguments& InArgs, const TWeakObjectPtr<UInternationalizationSettingsModel>& InSettingsModel, const TSharedRef<FLocalizedCulturesFlyweight>& InLocalizedCulturesFlyweight)
 		{
-			return true;
-		}
-		// (None) should appear before all else.
-		if(BText.IdenticalTo(LOCTEXT("None", "(None)")))
-		{
-			return false;
-		}
-		// Non-Specific Region should appear before all else.
-		if(AText.IdenticalTo(LOCTEXT("NoSpecificRegionOption", "Non-Specific Region")))
-		{
-			return true;
-		}
-		// Non-Specific Region should appear before all else.
-		if(BText.IdenticalTo(LOCTEXT("NoSpecificRegionOption", "Non-Specific Region")))
-		{
-			return false;
-		}
-
-		// Compare native region strings.
-		return( GetCultureNativeRegionText( A ).CompareToCaseIgnored( GetCultureNativeRegionText( B ) ) ) < 0;
-	}
-};
+			SettingsModel = InSettingsModel;
+			LocalizedCulturesFlyweight = InLocalizedCulturesFlyweight;
 
-TSharedRef<IDetailCustomization> FInternationalizationSettingsModelDetails::MakeInstance()
-{
-	TSharedRef<FInternationalizationSettingsModelDetails> InternationalizationSettingsModelDetails = MakeShareable(new FInternationalizationSettingsModelDetails());
-	return InternationalizationSettingsModelDetails;
-}
-
-FInternationalizationSettingsModelDetails::~FInternationalizationSettingsModelDetails()
-{
-	check(Model.IsValid());
-	Model->OnSettingsChanged().RemoveAll(this);
-}
-
-void FInternationalizationSettingsModelDetails::UpdateInternalStateFromSettingsModel()
-{
-	FInternationalization& I18N = FInternationalization::Get();
-
-	const FString SavedEditorCultureName = Model->GetEditorCultureName();
+			ChildSlot
+				[
+					SAssignNew(EditorCultureComboButton, SComboButton)
+					.ButtonContent()
+					[
+						SNew(STextBlock)
+						.Text(this, &SEditorCultureComboButton::GetContentText)
+					]
+				];
+			EditorCultureComboButton->SetOnGetMenuContent(FOnGetContent::CreateSP(this, &SEditorCultureComboButton::GetMenuContent));
+		}
 
-	// Attempt to use the culture specified in the settings.
-	if (!SavedEditorCultureName.IsEmpty())
-	{
-		const FCulturePtr SavedEditorCulture = I18N.GetCulture(SavedEditorCultureName);
-		const FCulturePtr SavedEditorLanguage = SavedEditorCulture.IsValid() ? I18N.GetCulture(SavedEditorCulture->GetTwoLetterISOLanguageName()) : nullptr;
+	private:
+		TWeakObjectPtr<UInternationalizationSettingsModel> SettingsModel;
+		TSharedPtr<FLocalizedCulturesFlyweight> LocalizedCulturesFlyweight;
+		TSharedPtr<SComboButton> EditorCultureComboButton;
 
-		// Accept the saved culture's language if it is available.
-		if (AvailableEditorLanguages.Contains(SavedEditorLanguage))
+	private:
+		FText GetContentText() const
 		{
-			SelectedEditorLanguage = SavedEditorLanguage;
-			RefreshAvailableEditorRegions();
-
-			// Accept the saved culture's region if it is available.
-			if (AvailableEditorRegions.Contains(SavedEditorCulture))
-			{
-				SelectedEditorCulture = SavedEditorCulture;
-			}
-			// Fallback on the first available region of the selected language.
-			else
+			bool IsConfigured = false;
+			FString EditorCultureName;
+			if (SettingsModel.IsValid())
 			{
-				SelectedEditorLanguage = (AvailableEditorRegions.Num() > 0) ? AvailableEditorRegions[0] : nullptr;
+				IsConfigured = SettingsModel->GetEditorCultureName(EditorCultureName);
 			}
+			FInternationalization& I18N = FInternationalization::Get();
+			const FCulturePtr Culture = IsConfigured ? I18N.GetCulture(EditorCultureName) : I18N.GetCurrentCulture();
+			return Culture.IsValid() ? FText::FromString(Culture->GetNativeName()) : LOCTEXT("None", "None");
 		}
-		else
-		{
-			SelectedEditorCulture = I18N.GetCurrentCulture();
-			SelectedEditorLanguage = SelectedEditorCulture.IsValid() ? I18N.GetCulture(SelectedEditorCulture->GetTwoLetterISOLanguageName()) : nullptr;
-			RefreshAvailableEditorRegions();
 
-			if (!SelectedEditorLanguage.IsValid())
+		TSharedRef<SWidget> GetMenuContent()
+		{
+			FCulturePtr EditorCulture;
 			{
-				// Fallback on the first available region of the selected language.
-				SelectedEditorLanguage = (AvailableEditorRegions.Num() > 0) ? AvailableEditorRegions[0] : nullptr;
+				FInternationalization& I18N = FInternationalization::Get();
+				bool IsConfigured = false;
+				FString EditorCultureName;
+				if (SettingsModel.IsValid())
+				{
+					IsConfigured = SettingsModel->GetEditorCultureName(EditorCultureName);
+				}
+				EditorCulture = IsConfigured ? I18N.GetCulture(EditorCultureName) : I18N.GetCurrentCulture();
 			}
-		}
-	}
-	// Fallback on the current culture of the editor.
-	else
-	{
-		SelectedEditorCulture = I18N.GetCurrentCulture();
-		SelectedEditorLanguage = SelectedEditorCulture.IsValid() ? I18N.GetCulture(SelectedEditorCulture->GetTwoLetterISOLanguageName()) : nullptr;
-		RefreshAvailableEditorRegions();
 
-		if (!SelectedEditorLanguage.IsValid())
-		{
-			// Fallback on the first available region of the selected language.
-			SelectedEditorLanguage = (AvailableEditorRegions.Num() > 0) ? AvailableEditorRegions[0] : nullptr;
+			const auto& OnSelectionChangedLambda = [=](FCulturePtr& SelectedCulture, ESelectInfo::Type SelectInfo)
+			{
+				if (SettingsModel.IsValid())
+				{
+					SettingsModel->SetEditorCultureName(SelectedCulture.IsValid() ? SelectedCulture->GetName() : TEXT(""));
+					if (SelectedCulture.IsValid())
+					{
+						FInternationalization& I18N = FInternationalization::Get();
+						I18N.SetCurrentCulture(SelectedCulture->GetName());
+					}
+				}
+				if (EditorCultureComboButton.IsValid())
+				{
+					EditorCultureComboButton->SetIsOpen(false);
+				}
+			};
+
+			const auto& IsCulturePickableLambda = [=](FCulturePtr Culture) -> bool
+			{
+				TArray<FString> CultureNames = Culture->GetPrioritizedParentCultureNames();
+				for (const FString& CultureName : CultureNames)
+				{
+					if (LocalizedCulturesFlyweight->LocalizedCulturesForEditor.Contains(Culture))
+					{
+						return true;
+					}
+				}
+				return false;
+			};
+
+			const auto& CulturePicker = SNew(SCulturePicker)
+				.InitialSelection(EditorCulture)
+				.OnSelectionChanged_Lambda(OnSelectionChangedLambda)
+				.IsCulturePickable_Lambda(IsCulturePickableLambda)
+				.UseNativeDisplayNames(true);
+
+			return SNew(SBox)
+				.MaxDesiredHeight(300.0f)
+				.WidthOverride(300.0f)
+				[
+					CulturePicker
+				];
 		}
-	}
-	
-	if (EditorLanguageComboBox.IsValid())
-	{
-		EditorLanguageComboBox->RefreshOptions();
-		EditorLanguageComboBox->SetSelectedItem(SelectedEditorLanguage);
-	}
+	};
 
-	if (EditorRegionComboBox.IsValid())
+	class SNativeGameCultureComboButton : public SCompoundWidget
 	{
-		EditorRegionComboBox->RefreshOptions();
-		EditorRegionComboBox->SetSelectedItem(SelectedEditorCulture);
-	}
+		SLATE_BEGIN_ARGS(SNativeGameCultureComboButton){}
+		SLATE_END_ARGS()
 
-	const FString SavedNativeGameCultureName = Model->GetNativeGameCultureName();
+	public:
+		void Construct(const FArguments& InArgs, const TWeakObjectPtr<UInternationalizationSettingsModel>& InSettingsModel, const TSharedRef<FLocalizedCulturesFlyweight>& InLocalizedCulturesFlyweight)
+		{
+			SettingsModel = InSettingsModel;
+			LocalizedCulturesFlyweight = InLocalizedCulturesFlyweight;
 
-	// Attempt to use the culture specified in the settings.
-	if (!SavedNativeGameCultureName.IsEmpty())
-	{
-		const FCulturePtr SavedNativeGameCulture = I18N.GetCulture(SavedNativeGameCultureName);
-		const FCulturePtr SavedNativeGameLanguage = SavedNativeGameCulture.IsValid() ? I18N.GetCulture(SavedNativeGameCulture->GetTwoLetterISOLanguageName()) : nullptr;
+			ChildSlot
+				[
+					SAssignNew(NativeGameCultureComboButton, SComboButton)
+					.ButtonContent()
+					[
+						SNew(STextBlock)
+						.Text(this, &SNativeGameCultureComboButton::GetContentText)
+					]
+				];
+			NativeGameCultureComboButton->SetOnGetMenuContent(FOnGetContent::CreateSP(this, &SNativeGameCultureComboButton::GetMenuContent));
+		}
 
-		// Accept the saved culture's language if it is available.
-		if (AvailableNativeGameLanguages.Contains(SavedNativeGameLanguage))
-		{
-			SelectedNativeGameLanguage = SavedNativeGameLanguage;
-			RefreshAvailableNativeGameRegions();
+	private:
+		TWeakObjectPtr<UInternationalizationSettingsModel> SettingsModel;
+		TSharedPtr<FLocalizedCulturesFlyweight> LocalizedCulturesFlyweight;
+		TSharedPtr<SComboButton> NativeGameCultureComboButton;
 
-			// Accept the saved culture's region if it is available.
-			if (AvailableNativeGameRegions.Contains(SavedNativeGameCulture))
+	private:
+		FText GetContentText() const
+		{
+			bool IsConfigured = false;
+			FString NativeGameCultureName;
+			if (SettingsModel.IsValid())
 			{
-				SelectedNativeGameCulture = SavedNativeGameCulture;
+				IsConfigured = SettingsModel->GetNativeGameCultureName(NativeGameCultureName);
 			}
-			// Fallback on the first available region of the selected language.
-			else
+			FCulturePtr Culture;
+			if (!NativeGameCultureName.IsEmpty())
 			{
-				SelectedNativeGameCulture = (AvailableNativeGameRegions.Num() > 0) ? AvailableNativeGameRegions[0] : nullptr;
+				FInternationalization& I18N = FInternationalization::Get();
+				Culture = I18N.GetCulture(NativeGameCultureName);
 			}
+			return Culture.IsValid() ? FText::FromString(Culture->GetDisplayName()) : LOCTEXT("None", "None");
 		}
-		else
-		{
-			SelectedNativeGameCulture = I18N.GetCurrentCulture();
-			SelectedNativeGameLanguage = SelectedNativeGameCulture.IsValid() ? I18N.GetCulture(SelectedNativeGameCulture->GetTwoLetterISOLanguageName()) : nullptr;
-			RefreshAvailableNativeGameRegions();
-
-			if (!SelectedNativeGameCulture.IsValid())
-			{
-				// Fallback on the first available region of the selected language.
-				SelectedNativeGameCulture = (AvailableNativeGameRegions.Num() > 0) ? AvailableNativeGameRegions[0] : nullptr;
-			}
-		}
-	}
-	// Fallback to no specified language or region.
-	else
-	{
-		SelectedNativeGameCulture = nullptr;
-		SelectedNativeGameLanguage = nullptr;
-		RefreshAvailableNativeGameRegions();
-	}
-
-	if (NativeGameLanguageComboBox.IsValid())
-	{
-		NativeGameLanguageComboBox->RefreshOptions();
-		NativeGameLanguageComboBox->SetSelectedItem(SelectedNativeGameLanguage);
-	}
-
-	if (NativeGameRegionComboBox.IsValid())
-	{
-		NativeGameRegionComboBox->RefreshOptions();
-		NativeGameRegionComboBox->SetSelectedItem(SelectedNativeGameCulture);
-	}
-
-	if (LocalizedPropertyNamesCheckBox.IsValid())
-	{
-		LocalizedPropertyNamesCheckBox->SetIsChecked(Model->ShouldLoadLocalizedPropertyNames() ? ECheckBoxState::Checked : ECheckBoxState::Unchecked);
-	}
-
-	if (UnlocalizedNodesAndPinsCheckBox.IsValid())
-	{
-		UnlocalizedNodesAndPinsCheckBox->SetIsChecked(Model->ShouldShowNodesAndPinsUnlocalized() ? ECheckBoxState::Unchecked : ECheckBoxState::Checked);
-	}
-}
-
-void FInternationalizationSettingsModelDetails::OnSettingsChanged()
-{
-	// If we made the changes, there's no need to update ourselves from the model.
-	if (IsMakingChangesToModel)
-	{
-		return;
-	}
-
-	UpdateInternalStateFromSettingsModel();
-}
-
-void FInternationalizationSettingsModelDetails::CustomizeDetails( IDetailLayoutBuilder& DetailBuilder )
-{
-	FInternationalization& I18N = FInternationalization::Get();
-
-	TArray< TWeakObjectPtr<UObject> > ObjectsBeingCustomized;
-	DetailBuilder.GetObjectsBeingCustomized(ObjectsBeingCustomized);
-	check(ObjectsBeingCustomized.Num() == 1);
 
-	if(ObjectsBeingCustomized[0].IsValid())
-	{
-		Model = Cast<UInternationalizationSettingsModel>(ObjectsBeingCustomized[0].Get());
-	}
-	check(Model.IsValid());
-
-	Model->OnSettingsChanged().AddRaw(this, &FInternationalizationSettingsModelDetails::OnSettingsChanged);
-
-	// Populate master list of editor cultures.
-	{
-		TArray<FCultureRef> LocalizedCultures;
-		// Use only the cultures for which the editor has localizations for.
-		FInternationalization::Get().GetCulturesWithAvailableLocalization(FPaths::GetEditorLocalizationPaths(), LocalizedCultures, true);
-		for (const FCultureRef& Culture : LocalizedCultures)
+		TSharedRef<SWidget> GetMenuContent()
 		{
-			AvailableEditorCultures.Add(Culture);
-			const FCulturePtr LanguageCulture = FInternationalization::Get().GetCulture(Culture->GetTwoLetterISOLanguageName());
-			if (LanguageCulture.IsValid())
+			FCulturePtr NativeGameCulture;
 			{
-				AvailableEditorLanguages.AddUnique(LanguageCulture);
+				bool IsConfigured = false;
+				FString NativeGameCultureName;
+				if (SettingsModel.IsValid())
+				{
+					IsConfigured = SettingsModel->GetNativeGameCultureName(NativeGameCultureName);
+				}
+				if (!NativeGameCultureName.IsEmpty())
+				{
+					FInternationalization& I18N = FInternationalization::Get();
+					NativeGameCulture = I18N.GetCulture(NativeGameCultureName);
+				}
 			}
-		}
-		AvailableEditorLanguages.Sort(FCompareCultureByNativeLanguage());
-	}
-	// Populate master list of game cultures.
-	{
-		TArray<FCultureRef> LocalizedCultures;
-		// Use only the cultures for which the game has localizations for.
-		FInternationalization::Get().GetCulturesWithAvailableLocalization(FPaths::GetGameLocalizationPaths(), LocalizedCultures, true);
-		for (const FCultureRef& Culture : LocalizedCultures)
-		{
-			AvailableNativeGameCultures.Add(Culture);
-			const FCulturePtr LanguageCulture = FInternationalization::Get().GetCulture(Culture->GetTwoLetterISOLanguageName());
-			AvailableNativeGameLanguages.AddUnique(LanguageCulture);
-		}
-		// Add an option for no game culture.
-		AvailableNativeGameCultures.Add(nullptr);
-		AvailableNativeGameLanguages.AddUnique(nullptr);
-		AvailableNativeGameLanguages.Sort(FCompareCultureByNativeLanguage());
-	}
-
-	UpdateInternalStateFromSettingsModel();
-
-	// If the saved editor culture is not the same as the current culture, a restart is needed to sync them fully and properly.
-	if (SelectedEditorCulture.IsValid() && SelectedEditorCulture != I18N.GetCurrentCulture())
-	{
-		RequiresRestart = true;
-	}
-
-	IDetailCategoryBuilder& CategoryBuilder = DetailBuilder.EditCategory("Internationalization");
-
-	{
-		const FText LanguageToolTipText = LOCTEXT("EditorLanguageTooltip", "Change which language's translations the editor uses. (Requires restart to take effect.)");
-
-		CategoryBuilder.AddCustomRow(LOCTEXT("EditorLanguageLabel", "Editor Localization Language"))
-			.NameContent()
-			[
-				SNew(SHorizontalBox)
-				+SHorizontalBox::Slot()
-				.Padding( FMargin( 0, 1, 0, 1 ) )
-				.FillWidth(1.0f)
-				[
-					SNew(STextBlock)
-					.Text(LOCTEXT("EditorLanguageLabel", "Editor Localization Language"))
-					.Font(DetailBuilder.GetDetailFont())
-					.ToolTipText(LanguageToolTipText)
-				]
-			]
-		.ValueContent()
-			.MaxDesiredWidth(300.0f)
-			[
-				SAssignNew(EditorLanguageComboBox, SComboBox< FCulturePtr > )
-				.OptionsSource( &AvailableEditorLanguages )
-				.InitiallySelectedItem(SelectedEditorLanguage)
-				.OnGenerateWidget(this, &FInternationalizationSettingsModelDetails::OnLanguageGenerateWidget, &DetailBuilder)
-				.ToolTipText(LanguageToolTipText)
-				.OnSelectionChanged(this, &FInternationalizationSettingsModelDetails::OnEditorLanguageSelectionChanged)
-				.Content()
-				[
-					SNew(STextBlock)
-					.Text(this, &FInternationalizationSettingsModelDetails::GetEditorCurrentLanguageText)
-					.Font(DetailBuilder.GetDetailFont())
-				]
-			];
 
-		const FText RegionToolTipText = LOCTEXT("EditorRegionTooltip", "Change which region's translations the editor uses. (Requires restart to take effect.)");
-
-		CategoryBuilder.AddCustomRow(LOCTEXT("EditorRegionLabel", "Editor Localization Region"))
-			.NameContent()
-			[
-				SNew(SHorizontalBox)
-				+SHorizontalBox::Slot()
-				.Padding( FMargin( 0, 1, 0, 1 ) )
-				.FillWidth(1.0f)
-				[
-					SNew(STextBlock)
-					.Text(LOCTEXT("EditorRegionLabel", "Editor Localization Region"))
-					.Font(DetailBuilder.GetDetailFont())
-					.ToolTipText(RegionToolTipText)
-				]
-			]
-		.ValueContent()
-			.MaxDesiredWidth(300.0f)
-			[
-				SAssignNew(EditorRegionComboBox, SComboBox< FCulturePtr > )
-				.OptionsSource( &AvailableEditorRegions )
-				.InitiallySelectedItem(SelectedEditorCulture)
-				.OnGenerateWidget(this, &FInternationalizationSettingsModelDetails::OnRegionGenerateWidget, &DetailBuilder)
-				.ToolTipText(RegionToolTipText)
-				.OnSelectionChanged(this, &FInternationalizationSettingsModelDetails::OnEditorRegionSelectionChanged)
-				.IsEnabled(this, &FInternationalizationSettingsModelDetails::IsEditorRegionSelectionAllowed)
-				.Content()
+			const auto& CulturePickerSelectLambda = [=](FCulturePtr& SelectedCulture, ESelectInfo::Type SelectInfo)
+			{
+				if (SettingsModel.IsValid())
+				{
+					SettingsModel->SetNativeGameCultureName(SelectedCulture.IsValid() ? SelectedCulture->GetName() : TEXT(""));
+					FTextLocalizationManager::Get().RefreshResources();
+				}
+				if (NativeGameCultureComboButton.IsValid())
+				{
+					NativeGameCultureComboButton->SetIsOpen(false);
+				}
+			};
+			const auto& CulturePickerIsPickableLambda = [=](FCulturePtr Culture) -> bool
+			{
+				TArray<FString> CultureNames = Culture->GetPrioritizedParentCultureNames();
+				for (const FString& CultureName : CultureNames)
+				{
+					if (LocalizedCulturesFlyweight->LocalizedCulturesForGame.Contains(Culture))
+					{
+						return true;
+					}
+				}
+				return false;
+			};
+			return SNew(SBox)
+				.MaxDesiredHeight(300.0f)
+				.WidthOverride(300.0f)
 				[
-					SNew(STextBlock)
-					.Text(this, &FInternationalizationSettingsModelDetails::GetCurrentEditorRegionText)
-					.Font(DetailBuilder.GetDetailFont())
-				]
-			];
-	}
+					SNew(SCulturePicker)
+					.InitialSelection(NativeGameCulture)
+					.OnSelectionChanged_Lambda(CulturePickerSelectLambda)
+					.IsCulturePickable_Lambda(CulturePickerIsPickableLambda)
+					.CanSelectNone(true)
+				];
+		}
+	};
+}
 
+void FInternationalizationSettingsModelDetails::CustomizeDetails(IDetailLayoutBuilder& DetailLayout)
+{
+	const TWeakObjectPtr<UInternationalizationSettingsModel> SettingsModel = [&]()
 	{
-		const FText LanguageToolTipText = LOCTEXT("GameLanguageTooltip", "Change which language the editor treats as native for game localizations. (Requires restart to take effect.)");
+		TArray< TWeakObjectPtr<UObject> > ObjectsBeingCustomized;
+		DetailLayout.GetObjectsBeingCustomized(ObjectsBeingCustomized);
+		check(ObjectsBeingCustomized.Num() == 1);
+		return Cast<UInternationalizationSettingsModel>(ObjectsBeingCustomized.Top().Get());
+	}();
 
-		CategoryBuilder.AddCustomRow(LOCTEXT("GameLanguageLabel", "Game Localization Language"))
-			.NameContent()
-			[
-				SNew(SHorizontalBox)
-				+SHorizontalBox::Slot()
-				.Padding( FMargin( 0, 1, 0, 1 ) )
-				.FillWidth(1.0f)
-				[
-					SNew(STextBlock)
-					.Text(LOCTEXT("GameLanguageLabel", "Game Localization Language"))
-					.Font(DetailBuilder.GetDetailFont())
-					.ToolTipText(LanguageToolTipText)
-				]
-			]
-		.ValueContent()
-			.MaxDesiredWidth(300.0f)
-			[
-				SAssignNew(NativeGameLanguageComboBox, SComboBox< FCulturePtr > )
-				.OptionsSource( &AvailableNativeGameLanguages )
-				.InitiallySelectedItem(SelectedNativeGameLanguage)
-				.OnGenerateWidget(this, &FInternationalizationSettingsModelDetails::OnLanguageGenerateWidget, &DetailBuilder)
-				.ToolTipText(LanguageToolTipText)
-				.OnSelectionChanged(this, &FInternationalizationSettingsModelDetails::OnNativeGameLanguageSelectionChanged)
-				.Content()
-				[
-					SNew(STextBlock)
-					.Text(this, &FInternationalizationSettingsModelDetails::GetNativeGameCurrentLanguageText)
-					.Font(DetailBuilder.GetDetailFont())
-				]
-			];
+	IDetailCategoryBuilder& DetailCategoryBuilder = DetailLayout.EditCategory("Internationalization", LOCTEXT("InternationalizationCategory", "Internationalization"));
 
-		const FText RegionToolTipText = LOCTEXT("GameRegionTooltip", "Change which region the editor treats as native for game localizations. (Requires restart to take effect.)");
+	const TSharedRef<FLocalizedCulturesFlyweight> LocalizedCulturesFlyweight = MakeShareable(new FLocalizedCulturesFlyweight());
 
-		CategoryBuilder.AddCustomRow(LOCTEXT("GameRegionLabel", "Game Localization Region"))
-			.NameContent()
-			[
-				SNew(SHorizontalBox)
-				+SHorizontalBox::Slot()
-				.Padding( FMargin( 0, 1, 0, 1 ) )
-				.FillWidth(1.0f)
-				[
-					SNew(STextBlock)
-					.Text(LOCTEXT("GameRegionLabel", "Game Localization Region"))
-					.Font(DetailBuilder.GetDetailFont())
-					.ToolTipText(RegionToolTipText)
-				]
-			]
-		.ValueContent()
-			.MaxDesiredWidth(300.0f)
-			[
-				SAssignNew(NativeGameRegionComboBox, SComboBox< FCulturePtr > )
-				.OptionsSource( &AvailableNativeGameRegions )
-				.InitiallySelectedItem(SelectedNativeGameCulture)
-				.OnGenerateWidget(this, &FInternationalizationSettingsModelDetails::OnRegionGenerateWidget, &DetailBuilder)
-				.ToolTipText(RegionToolTipText)
-				.OnSelectionChanged(this, &FInternationalizationSettingsModelDetails::OnNativeGameRegionSelectionChanged)
-				.IsEnabled(this, &FInternationalizationSettingsModelDetails::IsNativeGameRegionSelectionAllowed)
-				.Content()
-				[
-					SNew(STextBlock)
-					.Text(this, &FInternationalizationSettingsModelDetails::GetCurrentNativeGameRegionText)
-					.Font(DetailBuilder.GetDetailFont())
-				]
-			];
-	}
-
-	const FText FieldNamesToolTipText = LOCTEXT("EditorFieldNamesTooltip", "Toggle showing localized field names (requires restart to take effect)");
-
-	CategoryBuilder.AddCustomRow(LOCTEXT("EditorFieldNamesLabel", "Use Localized Field Names"))
-	.NameContent()
-	[
-		SNew(SHorizontalBox)
-		+SHorizontalBox::Slot()
-		.Padding( FMargin( 0, 1, 0, 1 ) )
-		.FillWidth(1.0f)
+	// Editor Culture Setting
+	const FText EditorCultureSettingDisplayName = LOCTEXT("EditorCultureSettingDisplayName", "Editor Culture");
+	const FText EditorCultureSettingToolTip = LOCTEXT("EditorCultureSettingToolTip", "The culture in which the editor should be displayed.");
+	FDetailWidgetRow& EditorCultureRow = DetailCategoryBuilder.AddCustomRow(EditorCultureSettingDisplayName);
+	EditorCultureRow.NameContent()
 		[
 			SNew(STextBlock)
-			.Text(LOCTEXT("EditorFieldNamesLabel", "Use Localized Field Names"))
-			.Font(DetailBuilder.GetDetailFont())
-			.ToolTipText(FieldNamesToolTipText)
-		]
-	]
-	.ValueContent()
-	.MaxDesiredWidth(300.0f)
-	[
-		SAssignNew(LocalizedPropertyNamesCheckBox, SCheckBox)
-		.IsChecked(Model->ShouldLoadLocalizedPropertyNames() ? ECheckBoxState::Checked : ECheckBoxState::Unchecked)
-		.ToolTipText(FieldNamesToolTipText)
-		.OnCheckStateChanged(this, &FInternationalizationSettingsModelDetails::ShouldLoadLocalizedFieldNamesCheckChanged)
-	];
-
-	const FText NodeAndPinsNamesToolTipText = LOCTEXT("GraphEditorNodesAndPinsLocalized_Tooltip", "Toggle localized node and pin titles in all graph editors");
-
-	CategoryBuilder.AddCustomRow(LOCTEXT("GraphEditorNodesAndPinsLocalized", "Use Localized Graph Editor Nodes and Pins"))
-	.NameContent()
-	[
-		SNew(SHorizontalBox)
-		+SHorizontalBox::Slot()
-		.Padding( FMargin( 0, 1, 0, 1 ) )
-		.FillWidth(1.0f)
+			.Text(EditorCultureSettingDisplayName)
+			.ToolTipText(EditorCultureSettingToolTip)
+			.Font(DetailLayout.GetDetailFont())
+		];
+	EditorCultureRow.ValueContent()
+		[
+			SNew(SEditorCultureComboButton, SettingsModel, LocalizedCulturesFlyweight)
+		];
+
+	// Native Game Culture Setting
+	const FText GameCultureSettingDisplayName = LOCTEXT("NativeGameCultureSettingDisplayName", "Native Game Culture");
+	const FText GameCultureSettingToolTip = LOCTEXT("NativeGameCultureSettingToolTip", "The culture in which game content is specified and localized from. WARNING: If this is not set to the culture the game is initially made in, modifying and saving assets may mix localized data in as native data!");
+	FDetailWidgetRow& GameCultureRow = DetailCategoryBuilder.AddCustomRow(GameCultureSettingDisplayName);
+	GameCultureRow.NameContent()
 		[
 			SNew(STextBlock)
-			.Text(LOCTEXT("GraphEditorNodesAndPinsLocalized", "Use Localized Graph Editor Nodes and Pins"))
-			.Font(DetailBuilder.GetDetailFont())
-			.ToolTipText(NodeAndPinsNamesToolTipText)
-		]
-	]
-	.ValueContent()
-	.MaxDesiredWidth(300.0f)
-	[
-		SAssignNew(UnlocalizedNodesAndPinsCheckBox, SCheckBox)
-		.IsChecked(Model->ShouldShowNodesAndPinsUnlocalized() ? ECheckBoxState::Unchecked : ECheckBoxState::Checked)
-		.ToolTipText(NodeAndPinsNamesToolTipText)
-		.OnCheckStateChanged(this, &FInternationalizationSettingsModelDetails::ShouldShowNodesAndPinsUnlocalized)
-	];
-
-	CategoryBuilder.AddCustomRow(LOCTEXT("EditorRestartWarningLabel", "RestartWarning"))
-	.Visibility( TAttribute<EVisibility>(this, &FInternationalizationSettingsModelDetails::GetInternationalizationRestartRowVisibility) )
-	.WholeRowContent()
-	.HAlign(HAlign_Center)
-	[
-		SNew( SHorizontalBox )
-		+SHorizontalBox::Slot()
-		.AutoWidth()
-		.VAlign(VAlign_Center)
-		.Padding(2.0f, 0.0f)
+			.Text(GameCultureSettingDisplayName)
+			.ToolTipText(GameCultureSettingToolTip)
+			.Font(DetailLayout.GetDetailFont())
+		];
+	GameCultureRow.ValueContent()
 		[
-			SNew( SImage )
-			.Image( FCoreStyle::Get().GetBrush("Icons.Warning") )
-		]
-		+SHorizontalBox::Slot()
-		.FillWidth(1.0f)
-		.VAlign(VAlign_Center)
+			SNew(SNativeGameCultureComboButton, SettingsModel, LocalizedCulturesFlyweight)
+		];
+
+	// Localized Field Names
+	const FText FieldNamesSettingDisplayName = LOCTEXT("EditorFieldNamesLabel", "Use Localized Field Names");
+	const FText FieldNamesSettingToolTip = LOCTEXT("EditorFieldNamesTooltip", "Toggle showing localized field names. NOTE: Requires restart to take effect.");
+	FDetailWidgetRow& FieldNamesRow = DetailCategoryBuilder.AddCustomRow(FieldNamesSettingDisplayName);
+	FieldNamesRow.NameContent()
 		[
-			SNew( STextBlock )
-			.Text( LOCTEXT("RestartWarningText", "Changes require restart to take effect.") )
-			.Font(DetailBuilder.GetDetailFont())
-		]
-	];
-}
-
-
-void FInternationalizationSettingsModelDetails::RefreshAvailableEditorRegions()
-{
-	AvailableEditorRegions.Empty();
-
-	// The available regions are based on the selected language, otherwise no regions are available.
-	if (SelectedEditorLanguage.IsValid())
-	{
-		const FString SelectedLanguageName = SelectedEditorLanguage->GetTwoLetterISOLanguageName();
-
-		// Add any cultures whose language is the the selected language.
-		for (const FCulturePtr& Culture : AvailableEditorCultures)
-		{
-			const FString CultureLanguageName = Culture->GetTwoLetterISOLanguageName();
-			if ( SelectedLanguageName == CultureLanguageName)
+			SNew(STextBlock)
+			.Text(FieldNamesSettingDisplayName)
+			.ToolTipText(FieldNamesSettingToolTip)
+			.Font(DetailLayout.GetDetailFont())
+		];
+	FieldNamesRow.ValueContent()
+		.MaxDesiredWidth(300.0f)
+		[
+			SNew(SCheckBox)
+			.IsChecked_Lambda([=]()
 			{
-				AvailableEditorRegions.AddUnique(Culture);
-			}
-		}
-
-		AvailableEditorRegions.Sort(FCompareCultureByNativeRegion());
-	}
-}
-
-FText FInternationalizationSettingsModelDetails::GetEditorCurrentLanguageText() const
-{
-	if (SelectedEditorLanguage.IsValid())
-	{
-		return FCompareCultureByNativeLanguage::GetCultureNativeLanguageText(SelectedEditorLanguage);
-	}
-	return LOCTEXT("None", "(None)");
-}
-
-void FInternationalizationSettingsModelDetails::OnEditorLanguageSelectionChanged( FCulturePtr Culture, ESelectInfo::Type SelectionType )
-{
-	TGuardValue<bool> Guard(IsMakingChangesToModel, true);
-	SelectedEditorLanguage = Culture;
-	RefreshAvailableEditorRegions();
-
-	// Fallback on the first available region of the selected language.
-	SelectedEditorCulture = (AvailableEditorRegions.Num() > 0) ? AvailableEditorRegions[0] : nullptr;
-
-	if (EditorRegionComboBox.IsValid())
-	{
-		EditorRegionComboBox->RefreshOptions();
-		EditorRegionComboBox->SetSelectedItem(SelectedEditorCulture);
-	}
-
-	Model->SetEditorCultureName(SelectedEditorCulture.IsValid() ? SelectedEditorCulture->GetName() : TEXT(""));
-	RequiresRestart = true;
-
-	ISettingsEditorModule& SettingsEditorModule = FModuleManager::GetModuleChecked<ISettingsEditorModule>("SettingsEditor");
-	SettingsEditorModule.OnApplicationRestartRequired();
-}
-
-FText FInternationalizationSettingsModelDetails::GetCurrentEditorRegionText() const
-{
-	if (SelectedEditorCulture.IsValid())
-	{
-		return FCompareCultureByNativeRegion::GetCultureNativeRegionText(SelectedEditorCulture);
-	}
-	return LOCTEXT("None", "(None)");
-}
-
-void FInternationalizationSettingsModelDetails::OnEditorRegionSelectionChanged( FCulturePtr Culture, ESelectInfo::Type SelectionType )
-{
-	TGuardValue<bool> Guard(IsMakingChangesToModel, true);
-	SelectedEditorCulture = Culture;
-
-	Model->SetEditorCultureName(SelectedEditorCulture.IsValid() ? SelectedEditorCulture->GetName() : TEXT(""));
-	RequiresRestart = true;
-
-	ISettingsEditorModule& SettingsEditorModule = FModuleManager::GetModuleChecked<ISettingsEditorModule>("SettingsEditor");
-	SettingsEditorModule.OnApplicationRestartRequired();
-}
-
-bool FInternationalizationSettingsModelDetails::IsEditorRegionSelectionAllowed() const
-{
-	return SelectedEditorLanguage.IsValid();
-}
-
-void FInternationalizationSettingsModelDetails::RefreshAvailableNativeGameRegions()
-{
-	AvailableNativeGameRegions.Empty();
-
-	const FString SelectedLanguageName = SelectedNativeGameLanguage.IsValid() ? SelectedNativeGameLanguage->GetTwoLetterISOLanguageName() : TEXT("");
-
-	// Setup the region list
-	for (const FCulturePtr& Culture : AvailableNativeGameCultures)
-	{
-		const FString CultureLanguageName = Culture.IsValid() ? Culture->GetTwoLetterISOLanguageName() : TEXT("");
-		if ( SelectedLanguageName == CultureLanguageName)
-		{
-			AvailableNativeGameRegions.AddUnique(Culture);
-		}
-	}
-
-	AvailableNativeGameRegions.Sort(FCompareCultureByNativeRegion());
-}
-
-FText FInternationalizationSettingsModelDetails::GetNativeGameCurrentLanguageText() const
-{
-	if( SelectedNativeGameLanguage.IsValid() )
-	{
-		return FCompareCultureByNativeLanguage::GetCultureNativeLanguageText(SelectedNativeGameLanguage);
-	}
-	return LOCTEXT("None", "(None)");
-}
-
-void FInternationalizationSettingsModelDetails::OnNativeGameLanguageSelectionChanged( FCulturePtr Culture, ESelectInfo::Type SelectionType )
-{
-	TGuardValue<bool> Guard(IsMakingChangesToModel, true);
-	SelectedNativeGameLanguage = Culture;
-	RefreshAvailableNativeGameRegions();
-
-	// Fallback on the first available region of the selected language.
-	SelectedNativeGameCulture = AvailableEditorRegions[0];
-
-	if (NativeGameRegionComboBox.IsValid())
-	{
-		NativeGameRegionComboBox->RefreshOptions();
-		NativeGameRegionComboBox->SetSelectedItem(SelectedNativeGameCulture);
-	}
-
-	Model->SetNativeGameCultureName(SelectedNativeGameCulture.IsValid() ? SelectedNativeGameCulture->GetName() : TEXT(""));
-	RequiresRestart = true;
-
-	ISettingsEditorModule& SettingsEditorModule = FModuleManager::GetModuleChecked<ISettingsEditorModule>("SettingsEditor");
-	SettingsEditorModule.OnApplicationRestartRequired();
-}
-
-FText FInternationalizationSettingsModelDetails::GetCurrentNativeGameRegionText() const
-{
-	if( SelectedNativeGameCulture.IsValid() )
-	{
-		return FCompareCultureByNativeRegion::GetCultureNativeRegionText(SelectedNativeGameCulture);
-	}
-	return LOCTEXT("None", "(None)");
-}
-
-void FInternationalizationSettingsModelDetails::OnNativeGameRegionSelectionChanged( FCulturePtr Culture, ESelectInfo::Type SelectionType )
-{
-	TGuardValue<bool> Guard(IsMakingChangesToModel, true);
-	SelectedNativeGameCulture = Culture;
-
-	Model->SetNativeGameCultureName(SelectedNativeGameCulture.IsValid() ? SelectedNativeGameCulture->GetName() : TEXT(""));
-	RequiresRestart = true;
-
-	ISettingsEditorModule& SettingsEditorModule = FModuleManager::GetModuleChecked<ISettingsEditorModule>("SettingsEditor");
-	SettingsEditorModule.OnApplicationRestartRequired();
-}
-
-bool FInternationalizationSettingsModelDetails::IsNativeGameRegionSelectionAllowed() const
-{
-	return SelectedNativeGameLanguage.IsValid();
-}
-
-TSharedRef<SWidget> FInternationalizationSettingsModelDetails::OnLanguageGenerateWidget( FCulturePtr Culture, IDetailLayoutBuilder* DetailBuilder ) const
-{
-	return SNew(STextBlock)
-		.Text(FCompareCultureByNativeLanguage::GetCultureNativeLanguageText(Culture))
-		.Font(DetailBuilder->GetDetailFont());
-}
-
-TSharedRef<SWidget> FInternationalizationSettingsModelDetails::OnRegionGenerateWidget( FCulturePtr Culture, IDetailLayoutBuilder* DetailBuilder ) const
-{
-	return SNew(STextBlock)
-		.Text(FCompareCultureByNativeRegion::GetCultureNativeRegionText(Culture))
-		.Font(DetailBuilder->GetDetailFont());
-}
-
-EVisibility FInternationalizationSettingsModelDetails::GetInternationalizationRestartRowVisibility() const
-{
-	return RequiresRestart ? EVisibility::Visible : EVisibility::Collapsed;
-}
-
-void FInternationalizationSettingsModelDetails::ShouldLoadLocalizedFieldNamesCheckChanged(ECheckBoxState CheckState)
-{
-	TGuardValue<bool> Guard(IsMakingChangesToModel, true);
-	Model->ShouldLoadLocalizedPropertyNames(LocalizedPropertyNamesCheckBox->IsChecked());
-}
-
-void FInternationalizationSettingsModelDetails::ShouldShowNodesAndPinsUnlocalized(ECheckBoxState CheckState)
-{
-	TGuardValue<bool> Guard(IsMakingChangesToModel, true);
-	Model->ShouldShowNodesAndPinsUnlocalized(CheckState == ECheckBoxState::Unchecked);
-
-	// Find all Schemas and force a visualization cache clear
-	for ( TObjectIterator<UClass> ClassIt; ClassIt; ++ClassIt )
-	{
-		UClass* CurrentClass = *ClassIt;
-
-		if (UEdGraphSchema* Schema = Cast<UEdGraphSchema>(CurrentClass->GetDefaultObject()))
-		{
-			Schema->ForceVisualizationCacheClear();
-		}
-	}
+				return SettingsModel.IsValid() && SettingsModel->ShouldLoadLocalizedPropertyNames() ? ECheckBoxState::Checked : ECheckBoxState::Unchecked;
+			})
+				.ToolTipText(FieldNamesSettingToolTip)
+				.OnCheckStateChanged_Lambda([=](ECheckBoxState State)
+			{
+				if (SettingsModel.IsValid())
+				{
+					SettingsModel->ShouldLoadLocalizedPropertyNames(State == ECheckBoxState::Checked);
+					FTextLocalizationManager::Get().RefreshResources();
+				}
+			})
+		];
+
+	// Localized Node and Pin Names
+	const FText NodeAndPinsNamesSettingDisplayName = LOCTEXT("GraphEditorNodesAndPinsLocalized", "Use Localized Graph Editor Nodes and Pins");
+	const FText NodeAndPinsNamesSettingToolTip = LOCTEXT("GraphEditorNodesAndPinsLocalized_Tooltip", "Toggle localized node and pin titles in all graph editors.");
+	FDetailWidgetRow& NodesAndPinsNameRow = DetailCategoryBuilder.AddCustomRow(NodeAndPinsNamesSettingDisplayName);
+	NodesAndPinsNameRow.NameContent()
+		[
+			SNew(STextBlock)
+			.Text(NodeAndPinsNamesSettingDisplayName)
+			.ToolTipText(NodeAndPinsNamesSettingToolTip)
+			.Font(DetailLayout.GetDetailFont())
+		];
+	NodesAndPinsNameRow.ValueContent()
+		.MaxDesiredWidth(300.0f)
+		[
+			SNew(SCheckBox)
+			.IsChecked_Lambda([=]()
+			{
+				return SettingsModel.IsValid() && SettingsModel->ShouldShowNodesAndPinsUnlocalized() ? ECheckBoxState::Unchecked : ECheckBoxState::Checked;
+			})
+				.ToolTipText(NodeAndPinsNamesSettingToolTip)
+				.OnCheckStateChanged_Lambda([=](ECheckBoxState State)
+			{
+				if (SettingsModel.IsValid())
+				{
+					SettingsModel->ShouldShowNodesAndPinsUnlocalized(State == ECheckBoxState::Unchecked);
+				}
+			})
+		];
 }
 
-#undef LOCTEXT_NAMESPACE
+#undef LOCTEXT_NAMESPACE
\ No newline at end of file
diff --git a/Engine/Source/Developer/InternationalizationSettings/Public/InternationalizationSettingsModelDetails.h b/Engine/Source/Developer/InternationalizationSettings/Public/InternationalizationSettingsModelDetails.h
--- a/Engine/Source/Developer/InternationalizationSettings/Public/InternationalizationSettingsModelDetails.h
+++ b/Engine/Source/Developer/InternationalizationSettings/Public/InternationalizationSettingsModelDetails.h
@@ -10,127 +10,6 @@ public:
 	/** Makes a new instance of this detail layout class for a specific detail view requesting it */
 	static TSharedRef<IDetailCustomization> MakeInstance();
 
-	/** Default Constructor */
-	FInternationalizationSettingsModelDetails()
-		: RequiresRestart(false)
-		, IsMakingChangesToModel(false)
-	{
-	}
-
-	/** Destructor */
-	~FInternationalizationSettingsModelDetails();
-
 	/** IDetailCustomization interface */
 	virtual void CustomizeDetails( IDetailLayoutBuilder& DetailLayout ) override;
-
-private:
-	void UpdateInternalStateFromSettingsModel();
-
-	void OnSettingsChanged();
-
-	/** Look for available editor culture regions */
-	void RefreshAvailableEditorRegions();
-
-	/** Delegate called to display the current editor language */
-	FText GetEditorCurrentLanguageText() const;
-
-	/** Delegate called when the editor language selection has changed */
-	void OnEditorLanguageSelectionChanged( FCulturePtr Culture, ESelectInfo::Type SelectionType );
-
-	/** Delegate called to display the current editor region */
-	FText GetCurrentEditorRegionText() const;
-
-	/** Delegate called when the editor region selection has changed */
-	void OnEditorRegionSelectionChanged( FCulturePtr Culture, ESelectInfo::Type SelectionType );
-
-	/** Delegate called to see if we're allowed to change the region */
-	bool IsEditorRegionSelectionAllowed() const;
-
-	/** Look for available game culture regions */
-	void RefreshAvailableNativeGameRegions();
-
-	/** Delegate called to display the current native game language */
-	FText GetNativeGameCurrentLanguageText() const;
-
-	/** Delegate called when the native game language selection has changed */
-	void OnNativeGameLanguageSelectionChanged( FCulturePtr Culture, ESelectInfo::Type SelectionType );
-
-	/** Delegate called to display the current native game region */
-	FText GetCurrentNativeGameRegionText() const;
-
-	/** Delegate called when the native game region selection has changed */
-	void OnNativeGameRegionSelectionChanged( FCulturePtr Culture, ESelectInfo::Type SelectionType );
-
-	/** Delegate called to see if we're allowed to change the region */
-	bool IsNativeGameRegionSelectionAllowed() const;
-
-	/** Generate a widget for the language combo */
-	TSharedRef<SWidget> OnLanguageGenerateWidget( FCulturePtr Culture, IDetailLayoutBuilder* DetailBuilder ) const;
-
-	/** Generate a widget for the region combo */
-	TSharedRef<SWidget> OnRegionGenerateWidget( FCulturePtr Culture, IDetailLayoutBuilder* DetailBuilder ) const;
-
-	/** Delegate called to determine whether to collapse or display the restart row. */
-	EVisibility GetInternationalizationRestartRowVisibility() const;
-
-	/** Delegate called when the the checked state of whether or not field names should be localized has changed. */
-	void ShouldLoadLocalizedFieldNamesCheckChanged(ECheckBoxState CheckState);
-
-	/** Delegate called when the the checked state of whether or not nodes and pins in graph editors should be localized has changed. */
-	void ShouldShowNodesAndPinsUnlocalized(ECheckBoxState CheckState);
-
-private:
-	TWeakObjectPtr<UInternationalizationSettingsModel> Model;
-
-	/** The culture selected at present */
-	FCulturePtr SelectedEditorCulture;
-
-	/** The language selected at present */
-	FCulturePtr SelectedEditorLanguage;
-
-	/** The cultures we have available to us */
-	TArray< FCulturePtr > AvailableEditorCultures;
-
-	/** The languages we have available to us */
-	TArray< FCulturePtr > AvailableEditorLanguages;
-
-	/** The dropdown for the languages */
-	TSharedPtr<SComboBox< FCulturePtr > > EditorLanguageComboBox;
-
-	/** The regions we have available to us */
-	TArray< FCulturePtr > AvailableEditorRegions;
-
-	/** The dropdown for the regions */
-	TSharedPtr<SComboBox< FCulturePtr > > EditorRegionComboBox;
-
-	/** The culture selected at present */
-	FCulturePtr SelectedNativeGameCulture;
-
-	/** The language selected at present */
-	FCulturePtr SelectedNativeGameLanguage;
-
-	/** The cultures we have available to us */
-	TArray< FCulturePtr > AvailableNativeGameCultures;
-
-	/** The languages we have available to us */
-	TArray< FCulturePtr > AvailableNativeGameLanguages;
-
-	/** The dropdown for the languages */
-	TSharedPtr<SComboBox< FCulturePtr > > NativeGameLanguageComboBox;
-
-	/** The regions we have available to us */
-	TArray< FCulturePtr > AvailableNativeGameRegions;
-
-	/** The dropdown for the regions */
-	TSharedPtr<SComboBox< FCulturePtr > > NativeGameRegionComboBox;
-
-	/** The check box for using localized property names */
-	TSharedPtr<SCheckBox> LocalizedPropertyNamesCheckBox;
-
-	/** The check box for showing nodes and pins unlocalized */
-	TSharedPtr<SCheckBox> UnlocalizedNodesAndPinsCheckBox;
-
-	bool RequiresRestart;
-
-	bool IsMakingChangesToModel;
 };
\ No newline at end of file
diff --git a/Engine/Source/Editor/EditorSettingsViewer/Private/EditorSettingsViewerModule.cpp b/Engine/Source/Editor/EditorSettingsViewer/Private/EditorSettingsViewerModule.cpp
--- a/Engine/Source/Editor/EditorSettingsViewer/Private/EditorSettingsViewerModule.cpp
+++ b/Engine/Source/Editor/EditorSettingsViewer/Private/EditorSettingsViewerModule.cpp
@@ -98,20 +98,12 @@ protected:
 		);
 
 		// region & language
-		ISettingsSectionPtr RegionAndLanguageSettingsSection = SettingsModule.RegisterSettings("Editor", "General", "Internationalization",
+		ISettingsSectionPtr RegionAndLanguageSettings = SettingsModule.RegisterSettings("Editor", "General", "Internationalization",
 			LOCTEXT("InternationalizationSettingsModelName", "Region & Language"),
 			LOCTEXT("InternationalizationSettingsModelDescription", "Configure the editor's behavior to use a language and fit a region's culture."),
 			GetMutableDefault<UInternationalizationSettingsModel>()
 		);
 
-		if (RegionAndLanguageSettingsSection.IsValid())
-		{
-			RegionAndLanguageSettingsSection->OnExport().BindRaw(this, &FEditorSettingsViewerModule::HandleRegionAndLanguageExport);
-			RegionAndLanguageSettingsSection->OnImport().BindRaw(this, &FEditorSettingsViewerModule::HandleRegionAndLanguageImport);
-			RegionAndLanguageSettingsSection->OnSaveDefaults().BindRaw(this, &FEditorSettingsViewerModule::HandleRegionAndLanguageSaveDefaults);
-			RegionAndLanguageSettingsSection->OnResetDefaults().BindRaw(this, &FEditorSettingsViewerModule::HandleRegionAndLanguageResetToDefault);
-		}
-
 		// input bindings
 		TWeakPtr<SWidget> InputBindingEditorPanel = FModuleManager::LoadModuleChecked<IInputBindingEditorModule>("InputBindingEditor").CreateInputBindingEditorPanel();
 		ISettingsSectionPtr InputBindingSettingsSection = SettingsModule.RegisterSettings("Editor", "General", "InputBindings",
@@ -368,39 +360,6 @@ private:
 		return true;
 	}
 
-	bool HandleRegionAndLanguageExport(const FString& FileName)
-	{
-		UInternationalizationSettingsModel* Model = GetMutableDefault<UInternationalizationSettingsModel>();
-		GConfig->Flush(false, Model->GetDefaultConfigFilename());
-		return BackupFile(Model->GetDefaultConfigFilename(), FileName);
-	}
-
-	bool HandleRegionAndLanguageImport(const FString& FileName)
-	{
-		UInternationalizationSettingsModel* Model = GetMutableDefault<UInternationalizationSettingsModel>();
-		if( EAppReturnType::Ok == ShowRestartWarning(LOCTEXT("ImportRegionAndLanguage_Title", "Import Region & Language")))
-		{
-			FUnrealEdMisc::Get().SetConfigRestoreFilename(FileName, Model->GetDefaultConfigFilename());
-			FUnrealEdMisc::Get().RestartEditor(false);
-
-			return true;
-		}
-
-		return false;
-	}
-
-	bool HandleRegionAndLanguageSaveDefaults()
-	{
-		GetMutableDefault<UInternationalizationSettingsModel>()->SaveDefaults();
-		return true;
-	}
-
-	bool HandleRegionAndLanguageResetToDefault()
-	{
-		GetMutableDefault<UInternationalizationSettingsModel>()->ResetToDefault();
-		return true;
-	}
-
 private:
 
 	/** Holds a pointer to the settings editor's view model. */
diff --git a/Engine/Source/Editor/Localization/Private/SCulturePicker.cpp b/Engine/Source/Editor/Localization/Private/SCulturePicker.cpp
new file mode 100644
--- /dev/null
+++ b/Engine/Source/Editor/Localization/Private/SCulturePicker.cpp
@@ -0,0 +1,292 @@
+// Copyright 1998-2015 Epic Games, Inc. All Rights Reserved.
+
+#include "LocalizationPrivatePCH.h"
+#include "SCulturePicker.h"
+#include "STableRow.h"
+#include "SSearchBox.h"
+
+#define LOCTEXT_NAMESPACE "CulturePicker"
+
+void SCulturePicker::Construct( const FArguments& InArgs )
+{
+	OnCultureSelectionChanged = InArgs._OnSelectionChanged;
+	IsCulturePickable = InArgs._IsCulturePickable;
+	UseNativeDisplayNames = InArgs._UseNativeDisplayNames;
+	CanSelectNone = InArgs._CanSelectNone;
+
+	TArray<FString> AllCultureNames;
+	FInternationalization::Get().GetCultureNames(AllCultureNames);
+	for (const auto& CultureName : AllCultureNames)
+	{
+		const FCulturePtr Culture = FInternationalization::Get().GetCulture(CultureName);
+		if (Culture.IsValid())
+		{
+			Cultures.AddUnique(Culture);
+			const TArray<FString> ParentCultureNames = Culture->GetPrioritizedParentCultureNames();
+			for (const auto& ParentCultureName : ParentCultureNames)
+			{
+				const FCulturePtr ParentCulture = FInternationalization::Get().GetCulture(ParentCultureName);
+				Cultures.AddUnique(ParentCulture);
+			}
+		}
+	}
+
+	BuildStockEntries();
+	RebuildEntries();
+
+	ChildSlot
+		[
+			SNew(SVerticalBox)
+			+SVerticalBox::Slot()
+			.AutoHeight()
+			[
+				SNew(SSearchBox)
+				.HintText( LOCTEXT("SearchHintText", "Name/Abbreviation") )
+				.OnTextChanged(this, &SCulturePicker::OnFilterStringChanged)
+				.DelayChangeNotificationsWhileTyping(true)
+			]
+			+SVerticalBox::Slot()
+			.FillHeight(1.0f)
+			[
+				SAssignNew(TreeView, STreeView<TSharedPtr<FCultureEntry>>)
+				.SelectionMode(ESelectionMode::Single)
+				.TreeItemsSource(&RootEntries)
+				.OnGenerateRow(this, &SCulturePicker::OnGenerateRow)
+				.OnGetChildren(this, &SCulturePicker::OnGetChildren)
+				.OnSelectionChanged(this, &SCulturePicker::OnSelectionChanged)
+			]
+		];
+
+
+	const auto& IsInitialSelection = [&](const TSharedPtr<FCultureEntry>& Entry) -> bool
+	{
+		return Entry->Culture == InArgs._InitialSelection;
+	};
+
+	const TSharedPtr<FCultureEntry>* InitialSelection = RootEntries.FindByPredicate(IsInitialSelection);
+	if (InitialSelection)
+	{
+		TGuardValue<bool> SupressSelectionGuard(SupressSelectionCallback, true);
+		TreeView->SetSelection(*InitialSelection);
+	}
+}
+
+void SCulturePicker::RequestTreeRefresh()
+{
+	RebuildEntries();
+	if (TreeView.IsValid())
+	{
+		TreeView->RequestTreeRefresh();
+	}
+}
+
+void SCulturePicker::BuildStockEntries()
+{
+	TArray<FCulturePtr> StockCultures;
+	TArray<FString> StockCultureNames;
+	FInternationalization::Get().GetCultureNames(StockCultureNames);
+	for (const auto& CultureName : StockCultureNames)
+	{
+		const FCulturePtr Culture = FInternationalization::Get().GetCulture(CultureName);
+		if (Culture.IsValid())
+		{
+			StockCultures.AddUnique(Culture);
+			const TArray<FString> ParentCultureNames = Culture->GetPrioritizedParentCultureNames();
+			for (const auto& ParentCultureName : ParentCultureNames)
+			{
+				const FCulturePtr ParentCulture = FInternationalization::Get().GetCulture(ParentCultureName);
+				StockCultures.AddUnique(ParentCulture);
+			}
+		}
+	}
+
+	TMap< FString, TSharedPtr<FCultureEntry> > StockCultureEntries;
+	StockCultureEntries.Reserve(StockCultures.Num());
+
+	StockEntries.Empty(StockCultures.Num());
+	StockEntries.Reserve(StockCultureEntries.Num());
+	for (const FCulturePtr& Culture : StockCultures)
+	{
+		const FString Name = Culture->GetName();
+
+		TSharedRef<FCultureEntry> CultureEntry = MakeShareable( new FCultureEntry(Culture) );
+
+		// Add culture to entries.
+		StockCultureEntries.Add( Name, CultureEntry );
+
+		// Add entry to list of root stock entries - it will be removed if someone has it as a child.
+		StockEntries.Add( CultureEntry );
+	}
+
+	// Update all parent entries to know about their child entries.
+	for (auto& Pair : StockCultureEntries)
+	{
+		const TSharedPtr<FCultureEntry> Entry = Pair.Value;
+
+		const TArray<FString> ParentCultureNames = Entry->Culture->GetPrioritizedParentCultureNames();
+
+		TArray<FString> InvalidAncestors;
+		for (const FString& ParentName : ParentCultureNames)
+		{
+			// Find existing parent entry.
+			TSharedPtr<FCultureEntry> ParentEntry;
+			{
+				const TSharedPtr<FCultureEntry>* const PossibleParentEntry = StockCultureEntries.Find(ParentName);
+				if (PossibleParentEntry)
+				{
+					ParentEntry = *PossibleParentEntry;
+				}
+			}
+
+			// Skip unfound parents.
+			if (!ParentEntry.IsValid())
+			{
+				continue;
+			}
+
+			// Don't add self as child of self.
+			if (Entry == ParentEntry)
+			{
+				continue;
+			}
+
+			// Invalidate ancestors of this parent entry.
+			TArray<FString> ParentAncestorCultureNames = ParentEntry->Culture->GetPrioritizedParentCultureNames();
+			for (const auto& AncestorCultureName : ParentAncestorCultureNames)
+			{
+				// Don't treat this parent entry as a parent of itself.
+				if (AncestorCultureName != ParentName)
+				{
+					InvalidAncestors.Add(AncestorCultureName);
+				}
+			}
+
+			// Don't add ancestors of ancestors.
+			if (InvalidAncestors.Contains(ParentName))
+			{
+				continue;
+			}
+
+			// Add current entry to parent entry.
+			ParentEntry->Children.Add( MakeShareable( new FCultureEntry(*Entry) ) );
+
+			// Remove the current entry from the base entries, as it is a child.
+			StockEntries.Remove(Entry);
+		}
+	}
+
+	// Sort entries.
+	const auto& CultureEntryComparator = [this](const TSharedPtr<FCultureEntry>& LHS, const TSharedPtr<FCultureEntry>& RHS) -> bool
+	{
+		return UseNativeDisplayNames ?
+			LHS->Culture->GetNativeName() < RHS->Culture->GetNativeName()
+			:
+			LHS->Culture->GetDisplayName() < RHS->Culture->GetDisplayName();
+	};
+	StockEntries.Sort(CultureEntryComparator);
+}
+
+void SCulturePicker::RebuildEntries()
+{
+	RootEntries.Reset();
+
+	if (CanSelectNone)
+	{
+		RootEntries.Add(MakeShareable(new FCultureEntry(nullptr, true)));
+	}
+	
+	TFunctionRef<void (const TArray< TSharedPtr<FCultureEntry> >&, TArray< TSharedPtr<FCultureEntry> >&)> DeepCopyAndFilter
+		= [&](const TArray< TSharedPtr<FCultureEntry> >& InEntries, TArray< TSharedPtr<FCultureEntry> >& OutEntries)
+	{
+		for (const TSharedPtr<FCultureEntry>& InEntry : InEntries)
+		{
+			// Set pickable flag.
+			bool IsPickable = true;
+			if (IsCulturePickable.IsBound())
+			{
+				IsPickable = IsCulturePickable.Execute(InEntry->Culture);
+			}
+
+			TSharedRef<FCultureEntry> OutEntry = MakeShareable(new FCultureEntry(InEntry->Culture, IsPickable));
+
+			// Recurse to children.
+			DeepCopyAndFilter(InEntry->Children, OutEntry->Children);
+
+			bool IsFilteredOut = false;
+			if (!FilterString.IsEmpty())
+			{
+				const FString Name = OutEntry->Culture->GetName();
+				const FString DisplayName = UseNativeDisplayNames ? OutEntry->Culture->GetNativeName() : OutEntry->Culture->GetDisplayName();
+				IsFilteredOut = !Name.Contains(FilterString) && !DisplayName.Contains(FilterString);
+			}
+
+			// If has children, must be added. If it is not filtered and it is pickable, should be added.
+			if ( OutEntry->Children.Num() != 0 || (!IsFilteredOut && IsPickable) )
+			{
+				OutEntries.Add(OutEntry);
+			}
+		}	
+	};
+	DeepCopyAndFilter(StockEntries, RootEntries);
+}
+
+void SCulturePicker::OnFilterStringChanged(const FText& InFilterString)
+{
+	FilterString = InFilterString.ToString();
+	RequestTreeRefresh();
+}
+
+TSharedRef<ITableRow> SCulturePicker::OnGenerateRow(TSharedPtr<FCultureEntry> Entry, const TSharedRef<STableViewBase>& Table)
+{
+	return	SNew(STableRow< TSharedPtr<FCultureEntry> >, Table)
+			[
+				SNew(STextBlock)
+				.Text(Entry->Culture.IsValid() ? FText::FromString(UseNativeDisplayNames ? Entry->Culture->GetNativeName() : Entry->Culture->GetDisplayName()) : LOCTEXT("None", "None"))
+				.ToolTip(
+						SNew(SToolTip)
+						.Content()
+						[
+							SNew(STextBlock)
+							.Text(Entry->Culture.IsValid() ? FText::FromString(Entry->Culture->GetName()) : LOCTEXT("None", "None"))
+							.HighlightText( FText::FromString(FilterString) )
+						]
+						)
+				.HighlightText( FText::FromString(FilterString) )
+				.ColorAndOpacity( Entry->IsSelectable ? FSlateColor::UseForeground() : FSlateColor::UseSubduedForeground() )
+			];
+}
+
+void SCulturePicker::OnGetChildren(TSharedPtr<FCultureEntry> Entry, TArray< TSharedPtr<FCultureEntry> >& Children)
+{
+	// Add entries from children array.
+	for (const auto& Child : Entry->Children)
+	{
+		Children.Add(Child);
+
+		// Sort entries.
+		const auto& CultureEntryComparator = [this](const TSharedPtr<FCultureEntry>& LHS, const TSharedPtr<FCultureEntry>& RHS) -> bool
+		{
+			return UseNativeDisplayNames ? 
+				LHS->Culture->GetNativeName() < RHS->Culture->GetNativeName()
+				:
+				LHS->Culture->GetDisplayName() < RHS->Culture->GetDisplayName();
+		};
+		Children.Sort(CultureEntryComparator);
+	}
+}
+
+void SCulturePicker::OnSelectionChanged(TSharedPtr<FCultureEntry> Entry, ESelectInfo::Type SelectInfo)
+{
+	if (SupressSelectionCallback)
+	{
+		return;
+	}
+
+	// Don't count as selection if the entry isn't actually selectable but is part of the hierarchy.
+	if (Entry.IsValid() && Entry->IsSelectable)
+	{
+		OnCultureSelectionChanged.ExecuteIfBound( Entry->Culture, SelectInfo );
+	}
+}
+
+#undef LOCTEXT_NAMESPACE
\ No newline at end of file
diff --git a/Engine/Source/Editor/Localization/Public/SCulturePicker.h b/Engine/Source/Editor/Localization/Public/SCulturePicker.h
new file mode 100644
--- /dev/null
+++ b/Engine/Source/Editor/Localization/Public/SCulturePicker.h
@@ -0,0 +1,98 @@
+// Copyright 1998-2015 Epic Games, Inc. All Rights Reserved.
+
+#pragma once
+
+#include "SCompoundWidget.h"
+#include "SlateDelegates.h"
+#include "CulturePointer.h"
+
+struct FCultureEntry
+{
+	FCulturePtr Culture;
+	TArray< TSharedPtr<FCultureEntry> > Children;
+	bool IsSelectable;
+
+	FCultureEntry(const FCulturePtr& InCulture, const bool InIsSelectable = true)
+		: Culture(InCulture)
+		, IsSelectable(InIsSelectable)
+	{}
+
+	FCultureEntry(const FCultureEntry& Source)
+		: Culture(Source.Culture)
+		, IsSelectable(Source.IsSelectable)
+	{
+		Children.Reserve(Source.Children.Num());
+		for (const auto& Child : Source.Children)
+		{
+			Children.Add( MakeShareable( new FCultureEntry(*Child) ) );
+		}
+	}
+};
+
+class LOCALIZATION_API SCulturePicker : public SCompoundWidget
+{
+public:
+	/** A delegate type invoked when a selection changes somewhere. */
+	DECLARE_DELEGATE_RetVal_OneParam(bool, FIsCulturePickable, FCulturePtr);
+	typedef TSlateDelegates< FCulturePtr >::FOnSelectionChanged FOnSelectionChanged;
+
+public:
+	SLATE_BEGIN_ARGS( SCulturePicker )
+		: _UseNativeDisplayNames(false)
+		, _CanSelectNone(false)
+	{}
+		SLATE_EVENT( FOnSelectionChanged, OnSelectionChanged )
+		SLATE_EVENT( FIsCulturePickable, IsCulturePickable )
+		SLATE_ARGUMENT( FCulturePtr, InitialSelection )
+		SLATE_ARGUMENT(bool, UseNativeDisplayNames)
+		SLATE_ARGUMENT(bool, CanSelectNone)
+	SLATE_END_ARGS()
+
+	SCulturePicker()
+	: UseNativeDisplayNames(false)
+	, CanSelectNone(false)
+	, SupressSelectionCallback(false)
+	{}
+
+	void Construct( const FArguments& InArgs );
+	void RequestTreeRefresh();
+
+private:
+	void BuildStockEntries();
+	void RebuildEntries();
+
+	void OnFilterStringChanged(const FText& InFilterString);
+	TSharedRef<ITableRow> OnGenerateRow(TSharedPtr<FCultureEntry> Entry, const TSharedRef<STableViewBase>& Table);
+	void OnGetChildren(TSharedPtr<FCultureEntry> Entry, TArray< TSharedPtr<FCultureEntry> >& Children);
+	void OnSelectionChanged(TSharedPtr<FCultureEntry> Entry, ESelectInfo::Type SelectInfo);
+
+private:
+	TSharedPtr< STreeView< TSharedPtr<FCultureEntry> > > TreeView;
+
+	/* The provided cultures array. */
+	TArray<FCulturePtr> Cultures;
+
+	/* The top level culture entries for all possible stock cultures. */
+	TArray< TSharedPtr<FCultureEntry> > StockEntries;
+
+	/* The top level culture entries to be displayed in the tree view. */
+	TArray< TSharedPtr<FCultureEntry> > RootEntries;
+
+	/* The string by which to filter cultures names. */
+	FString FilterString;
+
+	/** Delegate to invoke when selection changes. */
+	FOnSelectionChanged OnCultureSelectionChanged;
+
+	/** Delegate to invoke to set whether a culture is "pickable". */
+	FIsCulturePickable IsCulturePickable;
+
+	/* Should the cultures use their display names displayed as localized in their own (IE: native) culture? */
+	bool UseNativeDisplayNames;
+
+	/** Should a null culture option be available? */
+	bool CanSelectNone;
+
+	/* Flags that the selection callback shouldn't be called when selecting - useful for initial selection, for instance. */
+	bool SupressSelectionCallback;
+};
\ No newline at end of file
diff --git a/Engine/Source/Editor/LocalizationDashboard/Private/SCulturePicker.cpp b/Engine/Source/Editor/LocalizationDashboard/Private/SCulturePicker.cpp
deleted file mode 100644
diff --git a/Engine/Source/Editor/LocalizationDashboard/Private/SCulturePicker.h b/Engine/Source/Editor/LocalizationDashboard/Private/SCulturePicker.h
deleted file mode 100644
diff --git a/Engine/Source/Runtime/Core/Private/Internationalization/TextLocalizationManager.cpp b/Engine/Source/Runtime/Core/Private/Internationalization/TextLocalizationManager.cpp
--- a/Engine/Source/Runtime/Core/Private/Internationalization/TextLocalizationManager.cpp
+++ b/Engine/Source/Runtime/Core/Private/Internationalization/TextLocalizationManager.cpp
@@ -662,6 +662,13 @@ void FTextLocalizationManager::UpdateFromLocalizationResource(const FString& Loc
 	UpdateFromLocalizations(LocalizationEntryTrackers);
 }
 
+void FTextLocalizationManager::RefreshResources()
+{
+	const bool ShouldLoadEditor = WITH_EDITOR;
+	const bool ShouldLoadGame = FApp::IsGame();
+	LoadLocalizationResourcesForCulture(FInternationalization::Get().GetCurrentCulture()->GetName(), ShouldLoadEditor, ShouldLoadGame);
+}
+
 void FTextLocalizationManager::OnCultureChanged()
 {
 	const bool ShouldLoadEditor = bIsInitialized && WITH_EDITOR;
diff --git a/Engine/Source/Runtime/Core/Public/Internationalization/TextLocalizationManager.h b/Engine/Source/Runtime/Core/Public/Internationalization/TextLocalizationManager.h
--- a/Engine/Source/Runtime/Core/Public/Internationalization/TextLocalizationManager.h
+++ b/Engine/Source/Runtime/Core/Public/Internationalization/TextLocalizationManager.h
@@ -179,6 +179,9 @@ public:
 	/** Updates display string entries and adds new display string entries based on localizations found in a specified localization resource. */
 	void UpdateFromLocalizationResource(const FString& LocalizationResourceFilePath);
 
+	/** Reloads resources for the current culture. */
+	void RefreshResources();
+
 	/**	Returns the current text revision number. This value can be cached when caching information from the text localization manager.
 	 *	If the revision does not match, cached information may be invalid and should be recached. */
 	int32 GetTextRevision() const { return TextRevisionCounter; }
